Build Storage tool - Added support for staging the tool (making a standalone version)

diff --git a/Engine/Programs/BuildStorageTool/BuildStorageTool.uproject b/Engine/Programs/BuildStorageTool/BuildStorageTool.uproject
new file mode 100644
--- /dev/null
+++ b/Engine/Programs/BuildStorageTool/BuildStorageTool.uproject
@@ -0,0 +1,11 @@
+{
+	"FileVersion": 3,
+	"EngineAssociation": "",
+	"Category": "",
+	"Description": "",
+	"TargetPlatforms": [
+		"Windows",
+		"Mac",
+		"Linux"
+	]
+}
diff --git a/Engine/Programs/BuildStorageTool/Config/DefaultGame.ini b/Engine/Programs/BuildStorageTool/Config/DefaultGame.ini
new file mode 100644
--- /dev/null
+++ b/Engine/Programs/BuildStorageTool/Config/DefaultGame.ini
@@ -0,0 +1,9 @@
+[Staging]
++AllowedDirectories=Engine/Content/Editor/Slate
++AllowedDirectories=Engine/Content/Slate
++AllowedConfigFiles=Engine/Config/BaseEditorPerProjectUserSettings.ini
++AllowedConfigFiles=BuildStorageTool/Config/DeployId.ini
+
+[/Script/UnrealEd.ProjectPackagingSettings]
+; Reset the ini section deny list to allow us to stage the storage server configurations for this tool
+IniSectionDenylist=
\ No newline at end of file
diff --git a/Engine/Source/Programs/BuildStorageTool/BuildStorageTool.Build.cs b/Engine/Source/Programs/BuildStorageTool/BuildStorageTool.Build.cs
--- a/Engine/Source/Programs/BuildStorageTool/BuildStorageTool.Build.cs
+++ b/Engine/Source/Programs/BuildStorageTool/BuildStorageTool.Build.cs
@@ -15,6 +15,7 @@ public class BuildStorageTool : ModuleRules
 				"Core",
 				"DesktopPlatform",
 				"Json",
+				"PakFile",
 				"Projects",
 				"Slate",
 				"SlateCore",
@@ -34,5 +35,21 @@ public class BuildStorageTool : ModuleRules
 		}
 
 		PublicIncludePathModuleNames.Add("Launch");
+
+		if (Target.Platform == UnrealTargetPlatform.Win64)
+		{
+			RuntimeDependencies.Add("$(EngineDir)/Binaries/Win64/zen.exe");
+			RuntimeDependencies.Add("$(EngineDir)/Binaries/DotNET/OidcToken/win-x64/OidcToken.exe");
+		}
+		else if (Target.Platform == UnrealTargetPlatform.Linux)
+		{
+			RuntimeDependencies.Add("$(EngineDir)/Binaries/Linux/zen");
+			RuntimeDependencies.Add("$(EngineDir)/Binaries/DotNET/OidcToken/linux-x64/OidcToken");
+		}
+		else if (Target.Platform == UnrealTargetPlatform.Mac)
+		{
+			RuntimeDependencies.Add("$(EngineDir)/Binaries/Mac/zen");
+			RuntimeDependencies.Add("$(EngineDir)/Binaries/DotNET/OidcToken/osx-x64/OidcToken");
+		}
 	}
 }
