Add Editor Slate SVG icons as UFS runtime dependencies for BuildStorageTool so that they are included in the pak file for it when staging.

diff --git a/Engine/Source/Programs/BuildStorageTool/BuildStorageTool.Build.cs b/Engine/Source/Programs/BuildStorageTool/BuildStorageTool.Build.cs
--- a/Engine/Source/Programs/BuildStorageTool/BuildStorageTool.Build.cs
+++ b/Engine/Source/Programs/BuildStorageTool/BuildStorageTool.Build.cs
@@ -51,5 +51,7 @@ public class BuildStorageTool : ModuleRules
 			RuntimeDependencies.Add("$(EngineDir)/Binaries/Mac/zen");
 			RuntimeDependencies.Add("$(EngineDir)/Binaries/DotNET/OidcToken/osx-x64/OidcToken");
 		}
+
+		RuntimeDependencies.Add("$(EngineDir)/Content/Editor/Slate/...*.svg", StagedFileType.UFS);
 	}
 }
