Fix for Build Storage Tool not including Slate icons on Linux

diff --git a/Engine/Source/Programs/BuildStorageTool/BuildStorageTool.Build.cs b/Engine/Source/Programs/BuildStorageTool/BuildStorageTool.Build.cs
--- a/Engine/Source/Programs/BuildStorageTool/BuildStorageTool.Build.cs
+++ b/Engine/Source/Programs/BuildStorageTool/BuildStorageTool.Build.cs
@@ -7,7 +7,7 @@ public class BuildStorageTool : ModuleRules
 	public BuildStorageTool(ReadOnlyTargetRules Target) : base(Target)
 	{
 		bTreatAsEngineModule = true;
-		
+
 		PrivateDependencyModuleNames.AddRange(
 			new string[]
 			{
@@ -53,5 +53,10 @@ public class BuildStorageTool : ModuleRules
 		}
 
 		RuntimeDependencies.Add("$(EngineDir)/Content/Editor/Slate/...*.svg", StagedFileType.UFS);
+
+		if (Target.IsInPlatformGroup(UnrealPlatformGroup.Linux))
+		{
+			RuntimeDependencies.Add("$(EngineDir)/Content/SlateFileDialogs/Icons/...*.png", StagedFileType.UFS);
+		}
 	}
 }
