Expanding the png icons we include within SlateFileDialogs because there are still some missing on Linux under the "Common" folder.

diff --git a/Engine/Source/Programs/BuildStorageTool/BuildStorageTool.Build.cs b/Engine/Source/Programs/BuildStorageTool/BuildStorageTool.Build.cs
--- a/Engine/Source/Programs/BuildStorageTool/BuildStorageTool.Build.cs
+++ b/Engine/Source/Programs/BuildStorageTool/BuildStorageTool.Build.cs
@@ -56,7 +56,7 @@ public class BuildStorageTool : ModuleRules
 
 		if (Target.IsInPlatformGroup(UnrealPlatformGroup.Linux))
 		{
-			RuntimeDependencies.Add("$(EngineDir)/Content/SlateFileDialogs/Icons/...*.png", StagedFileType.UFS);
+			RuntimeDependencies.Add("$(EngineDir)/Content/SlateFileDialogs/...*.png", StagedFileType.UFS);
 		}
 	}
 }
