Fix BuildStorageTool main loop to ensure correct return codes on Mac/Linux.

diff --git a/Engine/Source/Programs/BuildStorageTool/Private/BuildStorageTool.cpp b/Engine/Source/Programs/BuildStorageTool/Private/BuildStorageTool.cpp
--- a/Engine/Source/Programs/BuildStorageTool/Private/BuildStorageTool.cpp
+++ b/Engine/Source/Programs/BuildStorageTool/Private/BuildStorageTool.cpp
@@ -342,14 +342,22 @@ int WINAPI WinMain(_In_ HINSTANCE hCurrInstance, _In_opt_ HINSTANCE hPrevInstanc
 }
 
 #elif PLATFORM_LINUX
+int UnixMainWrapper(const TCHAR* CommandLine)
+{
+	return BuildStorageToolMain(CommandLine) ? 0 : 1;
+};
 int main(int argc, char* argv[])
 {
-	return CommonUnixMain(argc, argv, &BuildStorageToolMain);
+	return CommonUnixMain(argc, argv, &UnixMainWrapper);
 }
 #elif PLATFORM_MAC
+int32 MacMainWrapper(const TCHAR* CommandLine)
+{
+	return BuildStorageToolMain(CommandLine) ? 0 : 1;
+};
 int main(int argc, char* argv[])
 {
-	[MacProgramDelegate mainWithArgc : argc argv : argv programMain : BuildStorageToolMain programExit : FEngineLoop::AppExit] ;
+	[MacProgramDelegate mainWithArgc : argc argv : argv programMain : MacMainWrapper programExit : FEngineLoop::AppExit] ;
 }
 #else
 #error "Unsupported platform!"
