Add version number to the Build Storage Tool title bar.

diff --git a/Engine/Source/Programs/BuildStorageTool/Private/BuildStorageTool.cpp b/Engine/Source/Programs/BuildStorageTool/Private/BuildStorageTool.cpp
--- a/Engine/Source/Programs/BuildStorageTool/Private/BuildStorageTool.cpp
+++ b/Engine/Source/Programs/BuildStorageTool/Private/BuildStorageTool.cpp
@@ -12,6 +12,7 @@
 #include "Misc/MessageDialog.h"
 #include "Misc/Optional.h"
 #include "Misc/Timespan.h"
+#include "Modules/BuildVersion.h"
 #include "Modules/ModuleInterface.h"
 #include "Modules/ModuleManager.h"
 #include "RequiredProgramMainCPPInclude.h"
@@ -240,7 +241,13 @@ public:
 private:
 	FText GetWindowTitle()
 	{
-		return LOCTEXT("WindowTitle", "Unreal Build Storage Tool");
+		FText ToolName = LOCTEXT("WindowTitle", "Unreal Build Storage Tool");
+		FBuildVersion CurrentBuildVersion;
+		if (FBuildVersion::TryRead(FBuildVersion::GetDefaultFileName(), CurrentBuildVersion))
+		{
+			return FText::Format(LOCTEXT("WindowTitle_VersionFormat", "{0} ({1})"), ToolName, FText::AsNumber(CurrentBuildVersion.Changelist, &FNumberFormattingOptions::DefaultNoGrouping()));
+		}
+		return ToolName;
 	}
 
 	EAppReturnType::Type OnModalMessageDialog(EAppMsgCategory InMessageCategory, EAppMsgType::Type InMessage, const FText& InText, const FText& InTitle)
