Use preprocessor defined version info in BuildStorageTool instead of trying to use the UE BuildVersion info.

diff --git a/Engine/Source/Programs/BuildStorageTool/Private/BuildStorageTool.cpp b/Engine/Source/Programs/BuildStorageTool/Private/BuildStorageTool.cpp
--- a/Engine/Source/Programs/BuildStorageTool/Private/BuildStorageTool.cpp
+++ b/Engine/Source/Programs/BuildStorageTool/Private/BuildStorageTool.cpp
@@ -31,6 +31,7 @@
 #include "Widgets/Text/STextBlock.h"
 #include "ZenServiceInstanceManager.h"
 #include "BuildServiceInstanceManager.h"
+#include "Version/AppVersionDefines.h"
 
 #if PLATFORM_WINDOWS
 #include "Runtime/Launch/Resources/Windows/Resource.h"
@@ -242,11 +243,13 @@ private:
 	FText GetWindowTitle()
 	{
 		FText ToolName = LOCTEXT("WindowTitle", "Unreal Build Storage Tool");
+#if defined(BUILD_STORAGE_TOOL_CHANGELIST_STRING)
 		FBuildVersion CurrentBuildVersion;
 		if (FBuildVersion::TryRead(FBuildVersion::GetDefaultFileName(), CurrentBuildVersion))
 		{
-			return FText::Format(LOCTEXT("WindowTitle_VersionFormat", "{0} ({1})"), ToolName, FText::AsNumber(CurrentBuildVersion.Changelist, &FNumberFormattingOptions::DefaultNoGrouping()));
+			return FText::Format(LOCTEXT("WindowTitle_VersionFormat", "{0} ({1})"), ToolName, FText::FromString(BUILD_STORAGE_TOOL_CHANGELIST_STRING));
 		}
+#endif
 		return ToolName;
 	}
 
diff --git a/Engine/Source/Programs/BuildStorageTool/Private/Version/AppVersionDefines.h b/Engine/Source/Programs/BuildStorageTool/Private/Version/AppVersionDefines.h
new file mode 100644
--- /dev/null
+++ b/Engine/Source/Programs/BuildStorageTool/Private/Version/AppVersionDefines.h
@@ -0,0 +1,29 @@
+// Copyright Epic Games, Inc. All Rights Reserved.
+
+#pragma once
+
+// Macros for encoding strings
+#define VERSION_TEXT(x) TEXT(x)
+#define VERSION_STRINGIFY_2(x) VERSION_TEXT(#x)
+#define VERSION_STRINGIFY(x) VERSION_STRINGIFY_2(x)
+
+#define BUILD_STORAGE_TOOL_APPNAME TEXT("BuildStorageTool-Source")
+
+#ifdef BUILD_STORAGE_TOOL_PRERELEASE
+#undef BUILD_STORAGE_TOOL_APPNAME
+#define BUILD_STORAGE_TOOL_APPNAME TEXT("BuildStorageTool-PreRelease")
+#endif
+
+#ifdef BUILD_STORAGE_TOOL_RELEASE
+#undef BUILD_STORAGE_TOOL_APPNAME
+#define BUILD_STORAGE_TOOL_APPNAME TEXT("BuildStorageTool")
+#endif
+
+#ifndef BUILD_STORAGE_TOOL_CHANGELIST
+#define BUILD_STORAGE_TOOL_CHANGELIST_STRING TEXT("NoCL")
+#endif
+
+#ifdef BUILD_STORAGE_TOOL_CHANGELIST
+#define BUILD_STORAGE_TOOL_CHANGELIST_STRING TEXT("CL-") VERSION_STRINGIFY(BUILD_STORAGE_TOOL_CHANGELIST)
+#endif
+
