Fix bug where window title version info was gated on UE build version info being present.

diff --git a/Engine/Source/Programs/BuildStorageTool/Private/BuildStorageTool.cpp b/Engine/Source/Programs/BuildStorageTool/Private/BuildStorageTool.cpp
--- a/Engine/Source/Programs/BuildStorageTool/Private/BuildStorageTool.cpp
+++ b/Engine/Source/Programs/BuildStorageTool/Private/BuildStorageTool.cpp
@@ -244,13 +244,10 @@ private:
 	{
 		FText ToolName = LOCTEXT("WindowTitle", "Unreal Build Storage Tool");
 #if defined(BUILD_STORAGE_TOOL_CHANGELIST_STRING)
-		FBuildVersion CurrentBuildVersion;
-		if (FBuildVersion::TryRead(FBuildVersion::GetDefaultFileName(), CurrentBuildVersion))
-		{
-			return FText::Format(LOCTEXT("WindowTitle_VersionFormat", "{0} ({1})"), ToolName, FText::FromString(BUILD_STORAGE_TOOL_CHANGELIST_STRING));
-		}
-#endif
+		return FText::Format(LOCTEXT("WindowTitle_VersionFormat", "{0} ({1})"), ToolName, FText::FromString(BUILD_STORAGE_TOOL_CHANGELIST_STRING));
+#else
 		return ToolName;
+#endif
 	}
 
 	EAppReturnType::Type OnModalMessageDialog(EAppMsgCategory InMessageCategory, EAppMsgType::Type InMessage, const FText& InText, const FText& InTitle)
