Build Storage Tool - Added help dialog and parameters system to help drive user supplied configuration

diff --git a/Engine/Source/Programs/BuildStorageTool/Private/BuildStorageTool.cpp b/Engine/Source/Programs/BuildStorageTool/Private/BuildStorageTool.cpp
--- a/Engine/Source/Programs/BuildStorageTool/Private/BuildStorageTool.cpp
+++ b/Engine/Source/Programs/BuildStorageTool/Private/BuildStorageTool.cpp
@@ -22,6 +22,7 @@
 #include "SMessageDialog.h"
 #include "StandaloneRenderer.h"
 #include "BuildStorageToolStyle.h"
+#include "BuildStorageToolHelpWidget.h"
 #include "Tasks/Task.h"
 #include "Templates/SharedPointer.h"
 #include "Widgets/Input/SButton.h"
@@ -32,6 +33,7 @@
 #include "ZenServiceInstanceManager.h"
 #include "BuildServiceInstanceManager.h"
 #include "Version/AppVersionDefines.h"
+#include "Parameters/BuildStorageToolParametersBuilder.h"
 
 #if PLATFORM_WINDOWS
 #include "Runtime/Launch/Resources/Windows/Resource.h"
@@ -56,6 +58,8 @@
 #define IMPLEMENT_SIGNING_KEY_REGISTRATION()
 IMPLEMENT_APPLICATION(BuildStorageTool, "BuildStorageTool");
 
+DEFINE_LOG_CATEGORY(LogBuildStorageTool);
+
 static void OnRequestExit()
 {
 	RequestEngineExit(TEXT("BuildStorageTool Closed"));
@@ -74,6 +78,7 @@ class FBuildStorageToolApp
 	TSharedPtr<SNotificationItem> CompileNotification;
 	TSharedPtr<UE::Zen::FServiceInstanceManager> ZenServiceInstanceManager;
 	TSharedPtr<UE::Zen::Build::FServiceInstanceManager> BuildServiceInstanceManager;
+	const FBuildStorageToolParameters& ToolParameters;
 
 	std::atomic<bool> bLatentExclusiveOperationActive = false;
 
@@ -82,6 +87,41 @@ class FBuildStorageToolApp
 		FSlateApplication::Get().RequestDestroyWindow(Window.ToSharedRef());
 	}
 
+	void OnHelpCommandPressed()
+	{
+		FSlateApplication& SlateApplicaton = FSlateApplication::Get();
+
+		if (SlateApplicaton.GetActiveModalWindow() != nullptr)
+		{
+			return;
+		}
+
+		TSharedPtr<SWidget> ParentWidget = SlateApplicaton.GetUserFocusedWidget(0);
+		if (!ensure(ParentWidget))
+		{
+			return;
+		}
+
+		// Determine the position of the window so that it will spawn near the mouse, but not go off the screen.
+		const FVector2D CursorPos = FSlateApplication::Get().GetCursorPos();
+		FSlateRect Anchor(CursorPos.X, CursorPos.Y, CursorPos.X, CursorPos.Y);
+		FVector2D AdjustedSummonLocation = FSlateApplication::Get().CalculatePopupWindowPosition(Anchor, FVector2D(441, 537), true, FVector2D::ZeroVector, Orient_Horizontal);
+
+		TSharedPtr<SWindow> WindowDialog = SNew(SWindow)
+			.AutoCenter(EAutoCenter::None)
+			.ScreenPosition(AdjustedSummonLocation)
+			.SupportsMaximize(false)
+			.SupportsMinimize(false)
+			.SizingRule(ESizingRule::Autosized)
+			.Title(LOCTEXT("WindowHeader", "Help"));
+
+		WindowDialog->SetContent(SNew(SBuildStorageToolHelpWidget)
+			.ToolParameters(&ToolParameters)
+		);
+
+		FSlateApplication::Get().AddModalWindow(WindowDialog.ToSharedRef(), ParentWidget);
+	}
+
 	bool CanExecuteExclusiveAction()
 	{
 		return !bLatentExclusiveOperationActive;
@@ -120,6 +160,18 @@ class FBuildStorageToolApp
 			NAME_None,
 			bRecursivelySearchable);
 
+		MenuBuilder.AddMenuEntry(
+			LOCTEXT("Help", "Help"),
+			LOCTEXT("Help_ToolTip", "Show the help dialog"),
+			FSlateIcon(),
+			FUIAction(
+				FExecuteAction::CreateRaw( this, &FBuildStorageToolApp::OnHelpCommandPressed ),
+				FCanExecuteAction::CreateRaw(this, &FBuildStorageToolApp::CanExecuteExclusiveAction)
+			),
+			NAME_None,
+			EUserInterfaceActionType::Button
+		);
+
 		MenuBuilder.AddMenuEntry(
 			LOCTEXT("Exit", "Exit"),
 			LOCTEXT("Exit_ToolTip", "Exits the Storage Server UI"),
@@ -136,8 +188,8 @@ class FBuildStorageToolApp
 	}
 
 public:
-	FBuildStorageToolApp(FSlateApplication& InSlate)
-		: Slate(InSlate)
+	FBuildStorageToolApp(FSlateApplication& InSlate, const FBuildStorageToolParameters& InToolParameters)
+		: Slate(InSlate), ToolParameters(InToolParameters)
 	{
 		ZenServiceInstanceManager = MakeShared<UE::Zen::FServiceInstanceManager>();
 		BuildServiceInstanceManager = MakeShared<UE::Zen::Build::FServiceInstanceManager>();
@@ -300,6 +352,9 @@ int BuildStorageToolMain(const TCHAR* CmdLine)
 	// Initialize high DPI mode
 	FSlateApplication::InitHighDPI(true);
 
+	FBuildStorageToolParametersBuilder ParametersBuilder;
+	FBuildStorageToolParameters Parameters = ParametersBuilder.Build();
+	
 	{
 		// Create the platform slate application (what FSlateApplication::Get() returns)
 		TSharedRef<FSlateApplication> Slate = FSlateApplication::Create(MakeShareable(FPlatformApplicationMisc::CreateApplication()));
@@ -327,7 +382,7 @@ int BuildStorageToolMain(const TCHAR* CmdLine)
 			FAppStyle::SetAppStyleSet(FBuildStorageToolStyle::Get());
 
 			// Run the inner application loop
-			FBuildStorageToolApp App(Slate.Get());
+			FBuildStorageToolApp App(Slate.Get(), Parameters);
 			App.Run();
 
 			// Clean up the custom styles
diff --git a/Engine/Source/Programs/BuildStorageTool/Private/BuildStorageTool.h b/Engine/Source/Programs/BuildStorageTool/Private/BuildStorageTool.h
--- a/Engine/Source/Programs/BuildStorageTool/Private/BuildStorageTool.h
+++ b/Engine/Source/Programs/BuildStorageTool/Private/BuildStorageTool.h
@@ -1,3 +1,7 @@
 // Copyright Epic Games, Inc. All Rights Reserved.
 
 #pragma once
+
+#include "CoreMinimal.h"
+
+DECLARE_LOG_CATEGORY_EXTERN(LogBuildStorageTool, Log, All);
\ No newline at end of file
diff --git a/Engine/Source/Programs/BuildStorageTool/Private/BuildStorageToolHelpWidget.cpp b/Engine/Source/Programs/BuildStorageTool/Private/BuildStorageToolHelpWidget.cpp
new file mode 100644
--- /dev/null
+++ b/Engine/Source/Programs/BuildStorageTool/Private/BuildStorageToolHelpWidget.cpp
@@ -0,0 +1,82 @@
+// Copyright Epic Games, Inc. All Rights Reserved.
+
+#include "BuildStorageToolHelpWidget.h"
+
+#include "BuildStorageToolStyle.h"
+#include "Widgets/Input/SHyperlink.h"
+#include "Widgets/SBoxPanel.h"
+
+#include "Version/AppVersionDefines.h"
+
+#define LOCTEXT_NAMESPACE "BuildStorageToolHelpWidget"
+
+// TODO BC: data drive all the data for the help page
+void SBuildStorageToolHelpWidget::Construct(const FArguments& InArgs)
+{
+	ToolParameters = InArgs._ToolParameters.Get();
+
+	TSharedPtr<SVerticalBox> Contents;
+	ChildSlot
+		[
+			SNew(SBorder)
+			.BorderImage(FAppStyle::GetBrush("ToolPanel.GroupBorder"))
+			[
+				SAssignNew(Contents, SVerticalBox)
+			]
+		];
+
+	Contents->AddSlot()
+		.AutoHeight()
+		.HAlign(HAlign_Center)
+		.Padding(5)
+		[
+			SNew(STextBlock)
+			.Text(LOCTEXT("Header", "The Build Storage Tool helps you download builds of different kinds.\n"))
+		];
+
+	Contents->AddSlot()
+		.AutoHeight()
+		.HAlign(HAlign_Center)
+		.VAlign(VAlign_Center)
+		.Padding(5)
+		[
+			SNew(STextBlock)
+			.Text(LOCTEXT("Documentation", "Please refer to the documentation page(s):"))
+		];
+
+	for(const FDocumentationLink& DocLink : ToolParameters->GeneralParameters.HelpLinks)
+	{
+		Contents->AddSlot()
+			.AutoHeight()
+			.HAlign(HAlign_Center)
+			.VAlign(VAlign_Center)
+			.Padding(10)
+			[
+				SNew(SHyperlink)
+				.Style(FBuildStorageToolStyle::Get(), TEXT("NavigationHyperlink"))
+				.Text(FText::FromString(DocLink.Text))
+				.ToolTipText(FText::FromString(DocLink.Tooltip))
+				.OnNavigate_Lambda([&DocLink]() { FPlatformProcess::LaunchURL(*DocLink.Link, nullptr, nullptr); })
+			];
+	}
+
+	Contents->AddSlot()
+		.AutoHeight()
+		.HAlign(HAlign_Center)
+		.VAlign(VAlign_Center)
+		.Padding(10)
+		[
+			SNew(STextBlock)
+			.Text(
+				FText::Format(
+#if defined(BUILD_STORAGE_TOOL_CHANGELIST_STRING)
+						LOCTEXT("HelpDialog_VersionFormat", "Application version: {0}"), FText::FromString(BUILD_STORAGE_TOOL_CHANGELIST_STRING)
+#else
+						LOCTEXT("HelpDialog_UnknownVersion", "Application version: Unknown")
+#endif
+				)
+			)
+		];
+}
+
+#undef LOCTEXT_NAMESPACE
\ No newline at end of file
diff --git a/Engine/Source/Programs/BuildStorageTool/Private/BuildStorageToolHelpWidget.h b/Engine/Source/Programs/BuildStorageTool/Private/BuildStorageToolHelpWidget.h
new file mode 100644
--- /dev/null
+++ b/Engine/Source/Programs/BuildStorageTool/Private/BuildStorageToolHelpWidget.h
@@ -0,0 +1,18 @@
+// Copyright Epic Games, Inc. All Rights Reserved.
+
+#pragma once
+
+#include "Widgets/SCompoundWidget.h"
+#include "Parameters/BuildStorageToolParameters.h"
+
+class SBuildStorageToolHelpWidget final : public SCompoundWidget
+{
+public:
+	SLATE_BEGIN_ARGS(SBuildStorageToolHelpWidget) {}
+		SLATE_ATTRIBUTE(const FBuildStorageToolParameters*, ToolParameters)
+	SLATE_END_ARGS()
+
+	void Construct(const FArguments& InArgs);
+private:
+	const FBuildStorageToolParameters* ToolParameters;
+};
\ No newline at end of file
diff --git a/Engine/Source/Programs/BuildStorageTool/Private/Parameters/BuildStorageToolParameters.h b/Engine/Source/Programs/BuildStorageTool/Private/Parameters/BuildStorageToolParameters.h
new file mode 100644
--- /dev/null
+++ b/Engine/Source/Programs/BuildStorageTool/Private/Parameters/BuildStorageToolParameters.h
@@ -0,0 +1,34 @@
+// Copyright Epic Games, Inc. All Rights Reserved.
+#pragma once
+
+#include "UObject/Class.h"
+#include "BuildStorageToolParameters.generated.h"
+
+USTRUCT()
+struct FDocumentationLink
+{
+	GENERATED_BODY()
+
+	UPROPERTY()
+	FString Text;
+	
+	UPROPERTY()
+	FString Tooltip;
+
+	UPROPERTY()
+	FString Link;
+};
+
+USTRUCT()
+struct FGeneralParameters 
+{
+	GENERATED_BODY()
+
+	UPROPERTY()
+	TArray<FDocumentationLink> HelpLinks;
+};
+
+struct FBuildStorageToolParameters
+{
+	FGeneralParameters GeneralParameters;
+};
diff --git a/Engine/Source/Programs/BuildStorageTool/Private/Parameters/BuildStorageToolParametersBuilder.cpp b/Engine/Source/Programs/BuildStorageTool/Private/Parameters/BuildStorageToolParametersBuilder.cpp
new file mode 100644
--- /dev/null
+++ b/Engine/Source/Programs/BuildStorageTool/Private/Parameters/BuildStorageToolParametersBuilder.cpp
@@ -0,0 +1,67 @@
+// Copyright Epic Games, Inc. All Rights Reserved.
+
+#include "BuildStorageToolParametersBuilder.h"
+
+#include "CoreGlobals.h"
+#include "BuildStorageTool.h"
+#include "HAL/FileManager.h"
+#include "Misc/ConfigContext.h"
+#include "Misc/StringOutputDevice.h"
+
+FBuildStorageToolParametersBuilder::FBuildStorageToolParametersBuilder()
+{
+	
+}
+
+FBuildStorageToolParameters FBuildStorageToolParametersBuilder::Build()
+{	
+	FString IniFilename = TEXT("Engine");
+	BuildStorageToolConfig = GConfig->FindConfigFile(IniFilename);
+
+	FBuildStorageToolParameters Parameters;
+	Parameters.GeneralParameters = BuildGeneralParameters();
+	return Parameters;
+}
+
+FGeneralParameters FBuildStorageToolParametersBuilder::BuildGeneralParameters()
+{
+	const FConfigSection* Section = BuildStorageToolConfig->FindSection(TEXT("BuildStorageTool.General"));
+	FGeneralParameters Output;
+
+	if(Section != nullptr)
+	{
+		FStringOutputDevice Errors;
+		FGeneralParameters::StaticStruct()->ImportText(*SectionToText(*Section), &Output, nullptr, 0, &Errors, FGeneralParameters::StaticStruct()->GetName());
+
+		if(!Errors.IsEmpty())
+		{
+			// TODO: Do we want to do more then just log a error here? 
+			UE_LOG(LogBuildStorageTool, Error, TEXT("Error loading parameter file %s"), *Errors);
+		}
+	}
+
+	return Output;
+}
+
+FString FBuildStorageToolParametersBuilder::SectionToText(const FConfigSection& InSection) const
+{
+	TArray<FString> lines;
+	for(const TPair<FName, FConfigValue>& Item : InSection.Array())
+	{
+		FString Value = Item.Value.GetValue();
+
+		// If it's an array/map/struct, we only need to quote the key, otherwise quote key and value
+		if((Value.IsNumeric() && !Value.Equals(TEXT("-"))) || (Value.StartsWith(TEXT("(")) && Value.EndsWith(TEXT(")")) && !Item.Key.ToString().Contains(TEXT("Regex"), ESearchCase::IgnoreCase)))
+		{
+			lines.Add(FString::Printf(TEXT("\"%s\"=%s"), *Item.Key.ToString(), *Item.Value.GetValue()));
+		}
+		else
+		{
+			lines.Add(FString::Printf(TEXT("\"%s\"=\"%s\""), *Item.Key.ToString(), *Item.Value.GetValue()));
+		}
+	}
+
+	FString FinalText = TEXT("(") + FString::Join(lines, TEXT(",")) + TEXT(")");
+	return FinalText;
+}
+
diff --git a/Engine/Source/Programs/BuildStorageTool/Private/Parameters/BuildStorageToolParametersBuilder.h b/Engine/Source/Programs/BuildStorageTool/Private/Parameters/BuildStorageToolParametersBuilder.h
new file mode 100644
--- /dev/null
+++ b/Engine/Source/Programs/BuildStorageTool/Private/Parameters/BuildStorageToolParametersBuilder.h
@@ -0,0 +1,21 @@
+// Copyright Epic Games, Inc. All Rights Reserved.
+#pragma once
+
+#include "CoreMinimal.h"
+
+#include "Misc/ConfigCacheIni.h"
+#include "BuildStorageToolParameters.h"
+
+class FBuildStorageToolParametersBuilder
+{
+public:
+	FBuildStorageToolParametersBuilder();
+	FBuildStorageToolParameters Build();
+
+private:
+	FConfigFile* BuildStorageToolConfig;
+private:
+	FGeneralParameters BuildGeneralParameters();
+	
+	FString SectionToText(const FConfigSection& InSection) const;
+};
