PR #13697: BuildStorageTool - Add built-in output log window

diff --git a/Engine/Source/Programs/BuildStorageTool/BuildStorageTool.Build.cs b/Engine/Source/Programs/BuildStorageTool/BuildStorageTool.Build.cs
--- a/Engine/Source/Programs/BuildStorageTool/BuildStorageTool.Build.cs
+++ b/Engine/Source/Programs/BuildStorageTool/BuildStorageTool.Build.cs
@@ -21,7 +21,8 @@ public class BuildStorageTool : ModuleRules
 				"SlateCore",
 				"StandaloneRenderer",
 				"StorageServerWidgets",
-				"Zen"
+				"Zen",
+				"OutputLog"
 			});
 
 		if (Target.IsInPlatformGroup(UnrealPlatformGroup.Linux))
diff --git a/Engine/Source/Programs/BuildStorageTool/Private/BuildStorageTool.cpp b/Engine/Source/Programs/BuildStorageTool/Private/BuildStorageTool.cpp
--- a/Engine/Source/Programs/BuildStorageTool/Private/BuildStorageTool.cpp
+++ b/Engine/Source/Programs/BuildStorageTool/Private/BuildStorageTool.cpp
@@ -34,6 +34,9 @@
 #include "BuildServiceInstanceManager.h"
 #include "Version/AppVersionDefines.h"
 #include "Parameters/BuildStorageToolParametersBuilder.h"
+#include "OutputLogCreationParams.h"
+#include "OutputLogModule.h"
+#include "OutputLogSettings.h"
 
 #if PLATFORM_WINDOWS
 #include "Runtime/Launch/Resources/Windows/Resource.h"
@@ -136,6 +139,11 @@ class FBuildStorageToolApp
 				LOCTEXT( "FileMenu", "File" ),
 				LOCTEXT( "FileMenu_ToolTip", "Opens the file menu" ),
 				FOnGetContent::CreateRaw( this, &FBuildStorageToolApp::FillFileMenu ) );
+
+			MenuBuilder.AddPullDownMenu(
+				LOCTEXT("ToolsMenu", "Tools"),
+				LOCTEXT("ToolsMenu_ToolTip", "Opens the tools menu"),
+				FOnGetContent::CreateRaw(this, &FBuildStorageToolApp::FillToolsMenu) );
 		}
 
 		// Create the menu bar
@@ -145,6 +153,73 @@ class FBuildStorageToolApp
 		return MenuBarWidget;
 	}
 
+	void OpenLogWindow()
+	{
+		/*** Output Log Widget ***/
+		FOutputLogModule& OutputLogModule = FModuleManager::Get().LoadModuleChecked<FOutputLogModule>("OutputLog");
+
+		// hide the debug console
+		IConsoleVariable* CVar = IConsoleManager::Get().FindConsoleVariable(TEXT("OutputLogModule.HideConsole"));
+		if(ensure(CVar))
+		{
+			CVar->Set(true);
+		}
+
+		// setup OutputLog settings
+		UOutputLogSettings* Settings = GetMutableDefault<UOutputLogSettings>();
+		if(Settings)
+		{
+			Settings->CategoryColorizationMode = ELogCategoryColorizationMode::ColorizeWholeLine;
+		}
+
+		// Open new slate window
+		FSlateApplication::Get().AddWindowAsNativeChild(
+			SNew(SWindow)
+			.Title(LOCTEXT("LogWindowTitle", "Build Storage Tool Log"))
+			.ClientSize(FVector2D(1200.0f, 600.0f))
+			.ActivationPolicy(EWindowActivationPolicy::Always)
+			.SizingRule(ESizingRule::UserSized)
+			.IsTopmostWindow(false)
+			.FocusWhenFirstShown(false)
+			.SupportsMaximize(true)
+			.SupportsMinimize(true)
+			.HasCloseButton(true)
+			[
+				OutputLogModule.MakeOutputLogWidget(FOutputLogCreationParams())
+			],
+			Window.ToSharedRef()
+		);
+	}
+
+	TSharedRef<SWidget> FillToolsMenu()
+	{
+		const bool bCloseSelfOnly = false;
+		const bool bSearchable = false;
+		const bool bRecursivelySearchable = false;
+
+		FMenuBuilder MenuBuilder(true,
+			nullptr,
+			TSharedPtr<FExtender>(),
+			bCloseSelfOnly,
+			&FCoreStyle::Get(),
+			bSearchable,
+			NAME_None,
+			bRecursivelySearchable);
+
+		MenuBuilder.AddMenuEntry(
+					LOCTEXT("OpenLogWindow", "Open Log Window"),
+					LOCTEXT("OpenLogWindow_ToolTip", "Opens the Log Window"),
+					FSlateIcon(),
+					FUIAction(
+						FExecuteAction::CreateRaw( this, &FBuildStorageToolApp::OpenLogWindow )
+					),
+					NAME_None,
+					EUserInterfaceActionType::Button
+				);
+		
+		return MenuBuilder.MakeWidget();
+	}
+
 	TSharedRef<SWidget> FillFileMenu()
 	{
 		const bool bCloseSelfOnly = false;
@@ -342,6 +417,12 @@ int BuildStorageToolMain(const TCHAR* CmdLine)
 	// start up the main loop
 	GEngineLoop.PreInit(CmdLine);
 
+	// ensure that the backlog is enabled
+	if(GLog)
+	{
+		GLog->EnableBacklog(true);
+	}
+
 	FSystemWideCriticalSection SystemWideBuildStorageToolCritSec(TEXT("BuildStorageTool"));
 	if (!SystemWideBuildStorageToolCritSec.IsValid())
 	{
