Move the Help menu item into its own root menu and change its name to About.

diff --git a/Engine/Source/Programs/BuildStorageTool/Private/BuildStorageTool.cpp b/Engine/Source/Programs/BuildStorageTool/Private/BuildStorageTool.cpp
--- a/Engine/Source/Programs/BuildStorageTool/Private/BuildStorageTool.cpp
+++ b/Engine/Source/Programs/BuildStorageTool/Private/BuildStorageTool.cpp
@@ -90,7 +90,7 @@ class FBuildStorageToolApp
 		FSlateApplication::Get().RequestDestroyWindow(Window.ToSharedRef());
 	}
 
-	void OnHelpCommandPressed()
+	void OnAboutCommandPressed()
 	{
 		FSlateApplication& SlateApplicaton = FSlateApplication::Get();
 
@@ -116,7 +116,7 @@ class FBuildStorageToolApp
 			.SupportsMaximize(false)
 			.SupportsMinimize(false)
 			.SizingRule(ESizingRule::Autosized)
-			.Title(LOCTEXT("WindowHeader", "Help"));
+			.Title(LOCTEXT("HelpAboutHeader", "About"));
 
 		WindowDialog->SetContent(SNew(SBuildStorageToolHelpWidget)
 			.ToolParameters(&ToolParameters)
@@ -134,7 +134,6 @@ class FBuildStorageToolApp
 	{
 		FMenuBarBuilder MenuBuilder( nullptr );
 		{
-			// Control
 			MenuBuilder.AddPullDownMenu(
 				LOCTEXT( "FileMenu", "File" ),
 				LOCTEXT( "FileMenu_ToolTip", "Opens the file menu" ),
@@ -144,6 +143,11 @@ class FBuildStorageToolApp
 				LOCTEXT("ToolsMenu", "Tools"),
 				LOCTEXT("ToolsMenu_ToolTip", "Opens the tools menu"),
 				FOnGetContent::CreateRaw(this, &FBuildStorageToolApp::FillToolsMenu) );
+
+			MenuBuilder.AddPullDownMenu(
+				LOCTEXT( "HelpMenu", "Help" ),
+				LOCTEXT( "HelpMenu_ToolTip", "Opens the help menu" ),
+				FOnGetContent::CreateRaw( this, &FBuildStorageToolApp::FillHelpMenu ) );
 		}
 
 		// Create the menu bar
@@ -236,23 +240,41 @@ class FBuildStorageToolApp
 			bRecursivelySearchable);
 
 		MenuBuilder.AddMenuEntry(
-			LOCTEXT("Help", "Help"),
-			LOCTEXT("Help_ToolTip", "Show the help dialog"),
+			LOCTEXT("Exit", "Exit"),
+			LOCTEXT("Exit_ToolTip", "Exits the Storage Server UI"),
 			FSlateIcon(),
 			FUIAction(
-				FExecuteAction::CreateRaw( this, &FBuildStorageToolApp::OnHelpCommandPressed ),
+				FExecuteAction::CreateRaw( this, &FBuildStorageToolApp::ExitTool ),
 				FCanExecuteAction::CreateRaw(this, &FBuildStorageToolApp::CanExecuteExclusiveAction)
 			),
 			NAME_None,
 			EUserInterfaceActionType::Button
 		);
 
+		return MenuBuilder.MakeWidget();
+	}
+
+	TSharedRef<SWidget> FillHelpMenu()
+	{
+		const bool bCloseSelfOnly = false;
+		const bool bSearchable = false;
+		const bool bRecursivelySearchable = false;
+
+		FMenuBuilder MenuBuilder(true,
+			nullptr,
+			TSharedPtr<FExtender>(),
+			bCloseSelfOnly,
+			&FCoreStyle::Get(),
+			bSearchable,
+			NAME_None,
+			bRecursivelySearchable);
+
 		MenuBuilder.AddMenuEntry(
-			LOCTEXT("Exit", "Exit"),
-			LOCTEXT("Exit_ToolTip", "Exits the Storage Server UI"),
+			LOCTEXT("About", "About"),
+			LOCTEXT("About_ToolTip", "Show the about dialog"),
 			FSlateIcon(),
 			FUIAction(
-				FExecuteAction::CreateRaw( this, &FBuildStorageToolApp::ExitTool ),
+				FExecuteAction::CreateRaw( this, &FBuildStorageToolApp::OnAboutCommandPressed ),
 				FCanExecuteAction::CreateRaw(this, &FBuildStorageToolApp::CanExecuteExclusiveAction)
 			),
 			NAME_None,
