BuildStorageTool: Add an Explore Log Directory menu item.

#Synced-CL 49059893

diff --git a/Engine/Source/Programs/BuildStorageTool/Private/BuildStorageTool.cpp b/Engine/Source/Programs/BuildStorageTool/Private/BuildStorageTool.cpp
--- a/Engine/Source/Programs/BuildStorageTool/Private/BuildStorageTool.cpp
+++ b/Engine/Source/Programs/BuildStorageTool/Private/BuildStorageTool.cpp
@@ -159,6 +159,17 @@ class FBuildStorageToolApp
 		return MenuBarWidget;
 	}
 
+	void ExploreLogDirectory()
+	{
+		FString Path = FPaths::ConvertRelativePathToFull(FPaths::ProjectLogDir());
+		if (!Path.Len() || !IFileManager::Get().DirectoryExists(*Path))
+		{
+			return;
+		}
+
+		FPlatformProcess::ExploreFolder(*FPaths::GetPath(Path));
+	}
+
 	void OpenLogWindow()
 	{
 		/*** Output Log Widget ***/
@@ -212,6 +223,17 @@ class FBuildStorageToolApp
 			NAME_None,
 			bRecursivelySearchable);
 
+		MenuBuilder.AddMenuEntry(
+					LOCTEXT("ExploreLogDirectory", "Explore Log Directory"),
+					LOCTEXT("ExploreLogDirectory_ToolTip", "Explore the directory containing log files"),
+					FSlateIcon(),
+					FUIAction(
+						FExecuteAction::CreateRaw( this, &FBuildStorageToolApp::ExploreLogDirectory )
+					),
+					NAME_None,
+					EUserInterfaceActionType::Button
+				);
+
 		MenuBuilder.AddMenuEntry(
 					LOCTEXT("OpenLogWindow", "Open Log Window"),
 					LOCTEXT("OpenLogWindow_ToolTip", "Opens the Log Window"),
