BuildStorageTool: declare dependency in build files to platform SDK to address compile issue on Mac.

#Synced-CL 51039791

diff --git a/Engine/Source/Programs/BuildStorageTool/BuildStorageTool.Build.cs b/Engine/Source/Programs/BuildStorageTool/BuildStorageTool.Build.cs
--- a/Engine/Source/Programs/BuildStorageTool/BuildStorageTool.Build.cs
+++ b/Engine/Source/Programs/BuildStorageTool/BuildStorageTool.Build.cs
@@ -7,6 +7,7 @@ public class BuildStorageTool : ModuleRules
 	public BuildStorageTool(ReadOnlyTargetRules Target) : base(Target)
 	{
 		bTreatAsEngineModule = true;
+		bRequiresPlatformSDK = true;
 
 		PrivateDependencyModuleNames.AddRange(
 			new string[]
