BuildStorageTool: Basic loopback-only TCP server+client functionality to allow BuildStorageTool processes to send messages back and forth.  Implemented as a slimmed down veresion of what exists in the CookDirector/CookWorker TCP comms.

#Synced-CL 51289657

diff --git a/Engine/Source/Programs/BuildStorageTool/BuildStorageTool.Build.cs b/Engine/Source/Programs/BuildStorageTool/BuildStorageTool.Build.cs
--- a/Engine/Source/Programs/BuildStorageTool/BuildStorageTool.Build.cs
+++ b/Engine/Source/Programs/BuildStorageTool/BuildStorageTool.Build.cs
@@ -12,20 +12,21 @@ public class BuildStorageTool : ModuleRules
 		PrivateDependencyModuleNames.AddRange(
 			new string[]
 			{
-				"ApplicationCore",
 				"Analytics",
+				"ApplicationCore",
 				"Core",
 				"DesktopPlatform",
 				"Json",
+				"OutputLog",
 				"PakFile",
 				"Projects",
 				"Slate",
 				"SlateCore",
+				"Sockets",
 				"StandaloneRenderer",
 				"StorageServerWidgets",
 				"StudioTelemetry",
 				"Zen",
-				"OutputLog"
 			});
 
 		if (Target.IsInPlatformGroup(UnrealPlatformGroup.Linux))
@@ -38,7 +39,10 @@ public class BuildStorageTool : ModuleRules
 			);
 		}
 
-		PublicIncludePathModuleNames.Add("Launch");
+		PrivateIncludePathModuleNames.AddRange(
+			new string[] {
+				"Launch",
+			});
 
 		if (Target.Platform == UnrealTargetPlatform.Win64)
 		{
diff --git a/Engine/Source/Programs/BuildStorageTool/Private/BuildStorageToolIPC.cpp b/Engine/Source/Programs/BuildStorageTool/Private/BuildStorageToolIPC.cpp
new file mode 100644
--- /dev/null
+++ b/Engine/Source/Programs/BuildStorageTool/Private/BuildStorageToolIPC.cpp
@@ -0,0 +1,515 @@
+// Copyright Epic Games, Inc. All Rights Reserved.
+
+#include "BuildStorageToolIPC.h"
+#include "Algo/RemoveIf.h"
+#include "Containers/Ticker.h"
+#include "Dom/JsonObject.h"
+#include "HAL/PlatformProcess.h"
+#include "IPAddress.h"
+#include "Misc/FileHelper.h"
+#include "Misc/Paths.h"
+#include "Serialization/CompactBinarySerialization.h"
+#include "Serialization/CompactBinaryValidation.h"
+#include "Serialization/JsonReader.h"
+#include "Serialization/JsonSerializer.h"
+#include "Serialization/JsonWriter.h"
+#include "Sockets.h"
+#include "SocketSubsystem.h"
+
+namespace UE::BuildStorageTool
+{
+static constexpr uint32 BSTIPC_Magic = 0x42535449; // 'BSTI'
+static constexpr uint16 BSTIPC_Version = 1;
+
+#pragma pack(push, 1)
+struct FFrameHeader
+{
+	uint32 Magic;
+	uint16 Version;
+	uint32 PayloadLen;
+};
+#pragma pack(pop)
+
+FString GetRendezvousFilePath()
+{
+	return FPaths::Combine(FPlatformProcess::UserTempDir(), TEXT("BuildStorageTool.ipc.json"));
+}
+
+bool WriteRendezvousFileJson(const uint16 Port)
+{
+	TSharedRef<FJsonObject> RootObject = MakeShared<FJsonObject>();
+	RootObject->SetNumberField(TEXT("port"), static_cast<double>(Port));
+
+	FString OutputString;
+	const TSharedRef<TJsonWriter<>> Writer = TJsonWriterFactory<>::Create(&OutputString);
+	if (!FJsonSerializer::Serialize(RootObject, Writer))
+	{
+		return false;
+	}
+
+	// Atomic-ish write: write temp then move over final path.
+	const FString FinalPath = GetRendezvousFilePath();
+	const FString TempPath = FinalPath + TEXT(".tmp");
+
+	if (!FFileHelper::SaveStringToFile(OutputString, *TempPath))
+	{
+		return false;
+	}
+
+	return IFileManager::Get().Move(*FinalPath, *TempPath, /*bReplace*/ true, /*bEvenIfReadOnly*/ false);
+}
+
+bool ReadRendezvousFileJson(uint16& OutPort)
+{
+	OutPort = 0;
+
+	FString InputString;
+	if (!FFileHelper::LoadFileToString(InputString, *GetRendezvousFilePath()))
+	{
+		return false;
+	}
+
+	TSharedPtr<FJsonObject> RootObject;
+	const TSharedRef<TJsonReader<>> Reader = TJsonReaderFactory<>::Create(InputString);
+	if (!FJsonSerializer::Deserialize(Reader, RootObject) || !RootObject.IsValid())
+	{
+		return false;
+	}
+
+	double PortAsNumber = 0.0;
+	if (!RootObject->TryGetNumberField(TEXT("port"), PortAsNumber))
+	{
+		return false;
+	}
+
+	const int32 PortAsInt = static_cast<int32>(PortAsNumber);
+	if (PortAsInt <= 0 || PortAsInt > 65535)
+	{
+		return false;
+	}
+
+	OutPort = static_cast<uint16>(PortAsInt);
+	return true;
+}
+
+static FSocket* StartListeningForMessages(uint16& OutPort)
+{
+	OutPort = 0;
+
+	ISocketSubsystem* SocketSubsystem = ISocketSubsystem::Get(PLATFORM_SOCKETSUBSYSTEM);
+	if (!SocketSubsystem)
+	{
+		return nullptr;
+	}
+
+	bool bIsValid = false;
+	TSharedRef<FInternetAddr> BindAddress = SocketSubsystem->CreateInternetAddr();
+	BindAddress->SetIp(TEXT("127.0.0.1"), bIsValid);
+	BindAddress->SetPort(0); // Ephemeral port
+
+	if (!bIsValid)
+	{
+		return nullptr;
+	}
+
+	TUniquePtr<FSocket> ListenSocket(
+		SocketSubsystem->CreateSocket(NAME_Stream, TEXT("BuildStorageToolIPC"), BindAddress->GetProtocolType())
+	);
+
+	if (!ListenSocket.IsValid())
+	{
+		return nullptr;
+	}
+
+	ListenSocket->SetReuseAddr(false);
+	ListenSocket->SetNoDelay(true);
+	ListenSocket->SetNonBlocking(true);
+
+	const int32 MaxBacklog = 16;
+	if (!ListenSocket->Bind(*BindAddress) || !ListenSocket->Listen(MaxBacklog))
+	{
+		return nullptr;
+	}
+
+	TSharedRef<FInternetAddr> BoundAddress = SocketSubsystem->CreateInternetAddr();
+	ListenSocket->GetAddress(*BoundAddress);
+	OutPort = BoundAddress->GetPort();
+
+	return ListenSocket.Release();
+}
+
+static bool RecvAll(FSocket& Socket, void* Data, int32 Length)
+{
+	uint8* BytePtr = static_cast<uint8*>(Data);
+
+	int32 TotalRead = 0;
+	while (TotalRead < Length)
+	{
+		int32 BytesRead = 0;
+		if (!Socket.Recv(BytePtr + TotalRead, Length - TotalRead, BytesRead) || BytesRead <= 0)
+		{
+			return false;
+		}
+		TotalRead += BytesRead;
+	}
+
+	return true;
+}
+
+static bool SendAll(FSocket& Socket, const void* Data, const int32 Length)
+{
+	const uint8* BytePtr = static_cast<const uint8*>(Data);
+	int32 TotalSent = 0;
+
+	while (TotalSent < Length)
+	{
+		int32 BytesSent = 0;
+		if (!Socket.Send(BytePtr + TotalSent, Length - TotalSent, BytesSent) || BytesSent <= 0)
+		{
+			return false;
+		}
+		TotalSent += BytesSent;
+	}
+
+	return true;
+}
+
+void FMessageListenServer::FReceiveBuffer::Reset()
+{
+	MessageType.Invalidate();
+	Payload.Reset();
+	bParsedHeader = false;
+	BytesRead = 0;
+}
+
+FMessageListenServer::FMessageListenServer()
+: ShutdownEvent(EEventMode::ManualReset)
+{
+}
+
+bool FMessageListenServer::Start()
+{
+	uint16 ListenPort = 0;
+	ListenSocket = StartListeningForMessages(ListenPort);
+	if (ListenSocket)
+	{
+		if (ListenPort != 0)
+		{
+			bWroteRendezvousFile = WriteRendezvousFileJson(ListenPort);
+		}
+
+		WorkerThread = FThread(TEXT("BuildStorageToolIPC_Thread"), [this] { ThreadLoop(); }, 128 * 1024);
+		return true;
+	}
+	return false;
+}
+
+void FMessageListenServer::ShutdownJoin()
+{
+	ShutdownEvent->Trigger();
+	if (WorkerThread.IsJoinable())
+	{
+		WorkerThread.Join();
+	}
+	ShutdownEvent->Reset();
+}
+
+void FMessageListenServer::ThreadLoop()
+{
+	constexpr float TickPeriod = 1.f;
+	constexpr float MinSleepTime = 0.001f;
+	for (;;)
+	{
+		double StartTime = FPlatformTime::Seconds();
+
+		bool bReadReady;
+		while (ListenSocket->HasPendingConnection(bReadReady) && bReadReady)
+		{
+			if (FSocket* ClientSocket = ListenSocket->Accept(TEXT("Client Connection")))
+			{
+				ClientSocket->SetNonBlocking(true);
+				ClientStates.AddDefaulted_GetRef().Socket = ClientSocket;
+			}
+		}
+
+		TickCommunication();
+
+		double CurrentTime = FPlatformTime::Seconds();
+		float RemainingDuration = StartTime + TickPeriod - CurrentTime;
+		if (RemainingDuration > .001f)
+		{
+			uint32 WaitTimeMilliseconds = static_cast<uint32>(RemainingDuration * 1000);
+			if (ShutdownEvent->Wait(WaitTimeMilliseconds))
+			{
+				break;
+			}
+		}
+		else if (ShutdownEvent->Wait(0))
+		{
+			break;
+		}
+	}
+
+	// Cleanup
+	if (ListenSocket)
+	{
+		ListenSocket->Close();
+		ISocketSubsystem::Get()->DestroySocket(ListenSocket);
+		ListenSocket = nullptr;
+	}
+	for (FClientState& ClientState : ClientStates)
+	{
+		if (ClientState.Socket)
+		{
+			ClientState.Socket->Close();
+			ISocketSubsystem::Get()->DestroySocket(ClientState.Socket);
+			ClientState.Socket = nullptr;
+		}
+	}
+	ClientStates.Empty();
+}
+
+void FMessageListenServer::TickCommunication()
+{
+	for (FClientState& ClientState : ClientStates)
+	{
+		if (!ClientState.Socket)
+		{
+			continue;
+		}
+
+		TArray<FMarshalledMessage> Messages;
+		constexpr uint64 MaxPayloadLen = 16 * 1024 * 1024;
+		EConnectionStatus SocketStatus = TryReadPacket(*ClientState.Socket, ClientState.ReceiveBuffer, Messages, MaxPayloadLen);
+		if (SocketStatus != EConnectionStatus::Okay && SocketStatus != EConnectionStatus::Incomplete)
+		{
+			ClientState.Socket->Close();
+			ISocketSubsystem::Get()->DestroySocket(ClientState.Socket);
+			ClientState.Socket = nullptr;
+		}
+
+		if (!OnMessageReceived.IsBound())
+		{
+			continue;
+		}
+
+		for (const FMarshalledMessage& Message : Messages)
+		{
+			if (Message.MessageType == FUrlMessage::MessageType)
+			{
+				FUrlMessage UrlMessage;
+				if (UrlMessage.TryRead(Message.Object))
+				{
+					ExecuteOnGameThread(UE_SOURCE_LOCATION,
+						[this, UrlMessage = MoveTemp(UrlMessage)]
+						{
+							OnMessageReceived.Broadcast(UrlMessage);
+						});
+				}
+			}
+		}
+	}
+
+	ClientStates.SetNum(Algo::RemoveIf(ClientStates, [] (FClientState& ClientState)
+	{
+		return ClientState.Socket == nullptr;
+	}));
+}
+
+FMessageListenServer::EConnectionStatus FMessageListenServer::TryReadPacket(FSocket& Socket, FReceiveBuffer& Buffer,
+	TArray<FMarshalledMessage>& Messages, uint64 MaxPayloadLen)
+{
+	int32 BytesRead;
+	if (!Buffer.bParsedHeader)
+	{
+		// When reading the relatively small FrameHeader, we wait for its total size to be ready
+		// using Peek to reduce complexity in this function.
+		FFrameHeader MessageHeader;
+		bool bStillAlive = Socket.Recv(reinterpret_cast<uint8*>(&MessageHeader), sizeof(MessageHeader),
+			BytesRead, ESocketReceiveFlags::Peek);
+		check(BytesRead <= sizeof(FFrameHeader));
+		if (BytesRead < sizeof(FFrameHeader))
+		{
+			if (!bStillAlive)
+			{
+				return EConnectionStatus::Terminated;
+			}
+			else
+			{
+				return EConnectionStatus::Okay;
+			}
+		}
+
+		Socket.Recv(reinterpret_cast<uint8*>(&MessageHeader), sizeof(MessageHeader), BytesRead,
+			ESocketReceiveFlags::None);
+		check(BytesRead == sizeof(FFrameHeader));
+		if (MessageHeader.Magic != BSTIPC_Magic || MessageHeader.Version != BSTIPC_Version ||
+			(MaxPayloadLen > 0 && MessageHeader.PayloadLen > MaxPayloadLen))
+		{
+			return EConnectionStatus::FormatError;
+		}
+
+		Buffer.Payload = FUniqueBuffer::Alloc(MessageHeader.PayloadLen);
+		Buffer.bParsedHeader = true;
+		Buffer.BytesRead = 0;
+	}
+
+
+	while (Buffer.Payload.GetSize() > Buffer.BytesRead)
+	{
+		// When reading the possibly large payload size that will block the remote socket
+		// from continuing until we have read some of it, we read as much of the payload
+		// as is available and store it in our dynamic and larger buffer in Buffer.Payload.
+		uint64 RemainingSize = Buffer.Payload.GetSize() - Buffer.BytesRead;
+		// Avoid possible OS restrictions on maximum read size by imposing our own moderate
+		// size per call to Recv.
+		constexpr uint64 MaxReadSize = 1000 * 1000 * 64;
+		int32 ReadSize = static_cast<int32>(FMath::Min(RemainingSize, MaxReadSize));
+		uint8* ReadData = static_cast<uint8*>(Buffer.Payload.GetData()) + Buffer.BytesRead;
+		bool bConnectionAlive = Socket.Recv(ReadData, ReadSize, BytesRead);
+		if (BytesRead <= 0)
+		{
+			if (!bConnectionAlive)
+			{
+				return EConnectionStatus::Terminated;
+			}
+			else
+			{
+				return EConnectionStatus::Okay;
+			}
+		}
+
+		check(BytesRead <= ReadSize);
+		Buffer.BytesRead += BytesRead;
+	}
+
+	// The FCbObjects we return have a pointer to this buffer and will keep it allocated
+	// until they are destructed.
+	FSharedBuffer SharedBuffer = Buffer.Payload.MoveToShared();
+	Buffer.Reset();
+	if (ValidateCompactBinary(SharedBuffer, ECbValidateMode::Default) != ECbValidateError::None)
+	{
+		return EConnectionStatus::FormatError;
+	}
+
+	FCbFieldIterator It = FCbFieldIterator::MakeRange(SharedBuffer);
+	while (It)
+	{
+		FMarshalledMessage& Message = Messages.Emplace_GetRef();
+		FCbFieldView MessageTypeView = *It;
+		Message.MessageType = MessageTypeView.AsUuid();
+		if (MessageTypeView.HasError())
+		{
+			return EConnectionStatus::FormatError;
+		}
+		++It;
+		FCbField ObjectField = *It;
+		Message.Object = ObjectField.AsObject();
+		if (ObjectField.HasError())
+		{
+			return EConnectionStatus::FormatError;
+		}
+		++It;
+	}
+	return EConnectionStatus::Okay;
+}
+
+FMessageClient::FMessageClient()
+: ShutdownEvent(EEventMode::ManualReset)
+{
+}
+
+bool FMessageClient::Start()
+{
+	if (!ReadRendezvousFileJson(Port))
+	{
+		return false;
+	}
+
+	ISocketSubsystem* SocketSubsystem = ISocketSubsystem::Get(PLATFORM_SOCKETSUBSYSTEM);
+	if (!SocketSubsystem)
+	{
+		return false;
+	}
+
+	bool bIsValid = false;
+	TSharedRef<FInternetAddr> RemoteAddress = SocketSubsystem->CreateInternetAddr();
+	RemoteAddress->SetIp(TEXT("127.0.0.1"), bIsValid);
+	RemoteAddress->SetPort(Port);
+
+	if (!bIsValid)
+	{
+		return false;
+	}
+
+	Socket = TUniquePtr<FSocket>(SocketSubsystem->CreateSocket(NAME_Stream, TEXT("BuildStorageToolIPC_Send"), RemoteAddress->GetProtocolType()));
+
+	if (!Socket.IsValid())
+	{
+		return false;
+	}
+
+	// Small retry window for early startup race
+	const double StartTime = FPlatformTime::Seconds();
+	while (!Socket->Connect(*RemoteAddress))
+	{
+		if (FPlatformTime::Seconds() - StartTime > 3.0)
+		{
+			return false;
+		}
+		FPlatformProcess::Sleep(0.02f);
+	}
+
+	bReady = true;
+	return true;
+}
+
+bool FMessageClient::SendSync(const IMessage& Message)
+{
+	if (!bReady || !Socket)
+	{
+		return false;
+	}
+
+	FFrameHeader Header;
+	Header.Magic = BSTIPC_Magic;
+	Header.Version = BSTIPC_Version;
+	FCbWriter Writer;
+	Writer.AddUuid(Message.GetMessageType());
+	Writer.BeginObject();
+	Message.Write(Writer);
+	Writer.EndObject();
+	uint64 PayloadLen64 = Writer.GetSaveSize();
+
+	if (PayloadLen64 > (MAX_int32 - sizeof(Header)))
+	{
+		return false;
+	}
+
+	Header.PayloadLen = static_cast<uint32>(PayloadLen64);
+
+	FUniqueBuffer SendBuffer = FUniqueBuffer::Alloc(sizeof(Header) + Header.PayloadLen);
+	FMemory::Memcpy(SendBuffer.GetData(), &Header, sizeof(Header));
+	FMutableMemoryView PayloadView(SendBuffer.GetView().RightChop(sizeof(Header)));
+	Writer.Save(PayloadView);
+
+	return SendAll(*Socket, SendBuffer.GetData(), static_cast<int32>(SendBuffer.GetSize()));
+}
+
+FGuid FUrlMessage::MessageType(TEXT("FC6BC9307BAB4D4C9F0C236FDEF2BE62"));
+
+FUrlMessage::FUrlMessage()
+{
+}
+
+void FUrlMessage::Write(FCbWriter& Writer) const
+{
+	Writer << "U" << Url;
+}
+
+bool FUrlMessage::TryRead(FCbObjectView Object)
+{
+	return LoadFromCompactBinary(Object["U"], Url);
+}
+
+} // namespace UE::BuildStorageTool
diff --git a/Engine/Source/Programs/BuildStorageTool/Private/BuildStorageToolIPC.h b/Engine/Source/Programs/BuildStorageTool/Private/BuildStorageToolIPC.h
new file mode 100644
--- /dev/null
+++ b/Engine/Source/Programs/BuildStorageTool/Private/BuildStorageToolIPC.h
@@ -0,0 +1,129 @@
+// Copyright Epic Games, Inc. All Rights Reserved.
+
+#pragma once
+
+#include "Delegates/Delegate.h"
+#include "HAL/Event.h"
+#include "HAL/Thread.h"
+#include "Misc/Guid.h"
+#include "Serialization/CompactBinary.h"
+#include "Serialization/CompactBinaryWriter.h"
+#include "Sockets.h"
+#include "Templates/UniquePtr.h"
+
+namespace UE::BuildStorageTool
+{
+
+FString GetRendezvousFilePath();
+bool WriteRendezvousFileJson(const uint16 Port);
+bool ReadRendezvousFileJson(uint16& OutPort);
+
+class IMessage
+{
+public:
+	virtual ~IMessage() = default;
+
+	/** Marshall the message to a CompactBinaryObject. */
+	virtual void Write(FCbWriter& Writer) const = 0;
+	/** Unmarshall the message from a CompactBinaryObject. */
+	virtual bool TryRead(FCbObjectView Object) = 0;
+	/** Return the Guid that identifies the message to the remote connection. */
+	virtual FGuid GetMessageType() const = 0;
+	/** Return the debugname for diagnostics. */
+	virtual const TCHAR* GetDebugName() const = 0;
+};
+
+typedef TMulticastDelegate<void(const IMessage& )> FMessageReceivedDelegate;
+
+class FMessageListenServer
+{
+public:
+	FMessageListenServer();
+
+	bool Start();
+	void ShutdownJoin();
+
+	FMessageReceivedDelegate OnMessageReceived;
+private:
+	enum class EConnectionStatus
+	{
+		/** Connection is still okay or operation succeeded. */
+		Okay,
+		/** Connection failed and no further use of the Socket is possible. */
+		Terminated,
+		/** The Socket is still active but the data received was invalid and recovery is not possible. */
+		FormatError,
+		/** The operation failed and the Socket is no longer usable. */
+		Failed,
+		/** The Socket is busy and the operation needs to be retried later */
+		Incomplete,
+	};
+
+	struct FReceiveBuffer
+	{
+	public:
+		void Reset();
+
+		FGuid MessageType;
+		FUniqueBuffer Payload;
+		uint64 BytesRead = 0;
+		bool bParsedHeader = false;
+	};
+
+	struct FMarshalledMessage
+	{
+		FGuid MessageType;
+		FCbObject Object;
+	};
+
+	void ThreadLoop();
+	void TickCommunication();
+	EConnectionStatus TryReadPacket(FSocket& Socket, FReceiveBuffer& Buffer,
+	TArray<FMarshalledMessage>& Messages, uint64 MaxPayloadLen);
+
+	struct FClientState
+	{
+		FSocket* Socket = nullptr;
+		FReceiveBuffer ReceiveBuffer;
+	};
+	TArray<FClientState> ClientStates;
+	FEventRef ShutdownEvent;
+	FThread WorkerThread;
+	FSocket* ListenSocket = nullptr;
+	uint16 Port = 0;
+	bool bWroteRendezvousFile = false;
+};
+
+class FMessageClient
+{
+public:
+	FMessageClient();
+
+	bool Start();
+	bool SendSync(const IMessage& Message);
+
+private:
+	void ThreadLoop();
+
+	FEventRef ShutdownEvent;
+	FThread WorkerThread;
+	TUniquePtr<FSocket> Socket;
+	uint16 Port = 0;
+	bool bReady = false;
+};
+
+struct FUrlMessage : public IMessage
+{
+public:
+	FUrlMessage();
+	virtual void Write(FCbWriter& Writer) const override;
+	virtual bool TryRead(FCbObjectView Object) override;
+	virtual FGuid GetMessageType() const override { return MessageType; }
+	virtual const TCHAR* GetDebugName() const override { return TEXT("URLMessage"); }
+
+public:
+	FString Url;
+	static FGuid MessageType;
+};
+
+} // namespace UE::BuildStorageTool
