BuildStorageTool: Adding app icon for Windows.

#Synced-CL 51529686

diff --git a/Engine/Source/Programs/BuildStorageTool/Resources/Windows/BuildStorageTool.rc b/Engine/Source/Programs/BuildStorageTool/Resources/Windows/BuildStorageTool.rc
new file mode 100644
--- /dev/null
+++ b/Engine/Source/Programs/BuildStorageTool/Resources/Windows/BuildStorageTool.rc
@@ -0,0 +1,83 @@
+// Copyright Epic Games, Inc. All Rights Reserved.
+
+#define APSTUDIO_READONLY_SYMBOLS
+/////////////////////////////////////////////////////////////////////////////
+//
+// Generated from the TEXTINCLUDE 2 resource.
+//
+#include <windows.h>
+#include "Runtime/Launch/Resources/Version.h"
+#include "Runtime/Launch/Resources/Windows/PCLaunch.rc"
+#include "Programs/BuildStorageTool/Private/Version/AppVersionDefines.h"
+
+#define VER_FILEVERSION         BUILD_STORAGE_TOOL_VERSION_MAJOR,BUILD_STORAGE_TOOL_VERSION_MINOR,BUILD_STORAGE_TOOL_VERSION_PATCH,0
+#define VER_FILEVERSION_STR     BUILD_STORAGE_TOOL_VERSION_STRING "\0"
+
+#define VER_PRODUCTVERSION      ENGINE_MAJOR_VERSION,ENGINE_MINOR_VERSION,ENGINE_PATCH_VERSION,0
+#define VER_PRODUCTVERSION_STR  ENGINE_VERSION_STRING "\0"
+
+#ifdef _DEBUG
+#define VER_FILEFLAGS           0x1L // VS_FF_DEBUG
+#else
+#define VER_FILEFLAGS           0x0L
+#endif
+
+#define IDI_BUILD_STORAGE_TOOL       IDICON_UEGame // 123
+
+/////////////////////////////////////////////////////////////////////////////
+#undef APSTUDIO_READONLY_SYMBOLS
+
+/////////////////////////////////////////////////////////////////////////////
+// English (U.S.) resources
+
+#if !defined(AFX_RESOURCE_DLL) || defined(AFX_TARG_ENU)
+#ifdef _WIN32
+LANGUAGE LANG_ENGLISH, SUBLANG_ENGLISH_US
+#pragma code_page(1252)
+#endif //_WIN32
+
+/////////////////////////////////////////////////////////////////////////////
+//
+// Version
+//
+
+VS_VERSION_INFO VERSIONINFO
+FILEVERSION     VER_FILEVERSION
+PRODUCTVERSION  VER_PRODUCTVERSION
+FILEFLAGSMASK   0x17L
+FILEFLAGS       VER_FILEFLAGS
+FILEOS          0x4L
+FILETYPE        0x2L
+FILESUBTYPE     0x0L
+BEGIN
+    BLOCK "StringFileInfo"
+    BEGIN
+        BLOCK "040904b0"
+        BEGIN
+            VALUE "CompanyName",      EPIC_COMPANY_NAME
+            VALUE "LegalCopyright",   EPIC_COPYRIGHT_STRING
+            VALUE "ProductName",      EPIC_PRODUCT_NAME
+            VALUE "ProductVersion",   VER_PRODUCTVERSION_STR
+            VALUE "FileDescription",  "Build Storage Tool"
+            VALUE "FileVersion",      VER_FILEVERSION_STR
+            VALUE "InternalName",     "BuildStorageTool"
+            VALUE "OriginalFilename", "BuildStorageTool.exe"
+        END
+    END
+    BLOCK "VarFileInfo"
+    BEGIN
+        VALUE "Translation", 0x409, 1200
+    END
+END
+
+/////////////////////////////////////////////////////////////////////////////
+//
+// Icon
+//
+
+// Icon with lowest ID value placed first to ensure application icon
+// remains consistent on all systems.
+IDI_BUILD_STORAGE_TOOL    ICON    "BuildStorageTool.ico"
+
+#endif    // English (U.S.) resources
+/////////////////////////////////////////////////////////////////////////////
diff --git a/Engine/Source/Programs/BuildStorageTool/Resources/Windows/make_ico.py b/Engine/Source/Programs/BuildStorageTool/Resources/Windows/make_ico.py
new file mode 100644
--- /dev/null
+++ b/Engine/Source/Programs/BuildStorageTool/Resources/Windows/make_ico.py
@@ -0,0 +1,17 @@
+# Copyright Epic Games, Inc. All Rights Reserved.
+
+from PIL import Image
+
+def MakeIcon(output):
+    folder = '.'
+    sizes = [256,128, 64, 48, 40, 32, 24, 16]
+    images = []
+    for size in sizes:
+        print(f"Loading {folder}/icon_{size}x{size}.png")
+        image = Image.open(f"{folder}/icon_{size}x{size}.png")
+        images.append(image)
+    print(f"Creating {output}")
+    images[0].save(f"{output}", append_images=images)
+
+MakeIcon("BuildStorageTool.ico")
+
