BuildStorageTool: Add defines used in RC file.

#Synced-CL 51529686

diff --git a/Engine/Source/Programs/BuildStorageTool/Private/Version/AppVersionDefines.h b/Engine/Source/Programs/BuildStorageTool/Private/Version/AppVersionDefines.h
--- a/Engine/Source/Programs/BuildStorageTool/Private/Version/AppVersionDefines.h
+++ b/Engine/Source/Programs/BuildStorageTool/Private/Version/AppVersionDefines.h
@@ -7,6 +7,12 @@
 #define VERSION_STRINGIFY_2(x) VERSION_TEXT(#x)
 #define VERSION_STRINGIFY(x) VERSION_STRINGIFY_2(x)
 
+#define BUILD_STORAGE_TOOL_VERSION_MAJOR 1
+#define BUILD_STORAGE_TOOL_VERSION_MINOR 0
+#define BUILD_STORAGE_TOOL_VERSION_PATCH 0
+
+#define BUILD_STORAGE_TOOL_VERSION_STRING "1.00"
+
 #define BUILD_STORAGE_TOOL_APPNAME TEXT("BuildStorageTool-Source")
 
 #ifdef BUILD_STORAGE_TOOL_PRERELEASE
