BuildStorageTool: Added download, paste, and Go to Build functionality for build artifacts using their download URL.  Also added quality of life improvements to locale formatting and tooltips for artifact size.

#Synced-CL 51529686

diff --git a/Engine/Source/Editor/StorageServerWidgets/Private/SBuildArtifactContentsViewer.cpp b/Engine/Source/Editor/StorageServerWidgets/Private/SBuildArtifactContentsViewer.cpp
--- a/Engine/Source/Editor/StorageServerWidgets/Private/SBuildArtifactContentsViewer.cpp
+++ b/Engine/Source/Editor/StorageServerWidgets/Private/SBuildArtifactContentsViewer.cpp
@@ -1,6 +1,7 @@
 // Copyright Epic Games, Inc. All Rights Reserved.
 
 #include "SBuildArtifactContentsViewer.h"
+#include "Internationalization/Internationalization.h"
 #include "Serialization/JsonWriter.h"
 #include "Widgets/Text/STextBlock.h"
 #include "Widgets/Layout/SBorder.h"
@@ -301,19 +302,42 @@ TSharedRef<ITableRow> SBuildArtifactContentsViewer::OnGenerateRow(TSharedPtr<FCo
 
 FString SBuildArtifactContentsViewer::FormatSize(uint64 SizeInBytes)
 {
+	const FCulturePtr LocaleCulture = FInternationalization::Get().GetCulture(FPlatformMisc::GetDefaultLocale());
+
 	if (SizeInBytes < 1024)
 	{
 		return FString::Printf(TEXT("%llu B"), SizeInBytes);
 	}
+
+	const TCHAR* Suffix;
+	double Value;
+	int32 FractionalDigits;
+
 	if (SizeInBytes < 1024 * 1024)
 	{
-		return FString::Printf(TEXT("%.1f KB"), (double)SizeInBytes / 1024.0);
+		Value = (double)SizeInBytes / 1024.0;
+		Suffix = TEXT("KB");
+		FractionalDigits = 1;
+	}
+	else if (SizeInBytes < 1024ull * 1024 * 1024)
+	{
+		Value = (double)SizeInBytes / (1024.0 * 1024.0);
+		Suffix = TEXT("MB");
+		FractionalDigits = 1;
 	}
-	if (SizeInBytes < 1024ull * 1024 * 1024)
+	else
 	{
-		return FString::Printf(TEXT("%.1f MB"), (double)SizeInBytes / (1024.0 * 1024.0));
+		Value = (double)SizeInBytes / (1024.0 * 1024.0 * 1024.0);
+		Suffix = TEXT("GB");
+		FractionalDigits = 2;
 	}
-	return FString::Printf(TEXT("%.2f GB"), (double)SizeInBytes / (1024.0 * 1024.0 * 1024.0));
+
+	const FNumberFormattingOptions Options = FNumberFormattingOptions()
+		.SetUseGrouping(false)
+		.SetMinimumFractionalDigits(FractionalDigits)
+		.SetMaximumFractionalDigits(FractionalDigits);
+
+	return FText::AsNumber(Value, &Options, LocaleCulture).ToString() + TEXT(" ") + Suffix;
 }
 
 void SBuildArtifactContentsViewer::SortAndRefreshList()
diff --git a/Engine/Source/Editor/StorageServerWidgets/Private/SBuildArtifactMetadataViewer.cpp b/Engine/Source/Editor/StorageServerWidgets/Private/SBuildArtifactMetadataViewer.cpp
--- a/Engine/Source/Editor/StorageServerWidgets/Private/SBuildArtifactMetadataViewer.cpp
+++ b/Engine/Source/Editor/StorageServerWidgets/Private/SBuildArtifactMetadataViewer.cpp
@@ -1,6 +1,7 @@
 // Copyright Epic Games, Inc. All Rights Reserved.
 
 #include "SBuildArtifactMetadataViewer.h"
+#include "Internationalization/Internationalization.h"
 #include "Widgets/Text/STextBlock.h"
 #include "Widgets/Layout/SScrollBox.h"
 #include "Widgets/Layout/SBorder.h"
@@ -192,7 +193,10 @@ void SBuildArtifactMetadataViewer::ParseMetadata(const FCbObject& Metadata)
 		}
 		else if (Field.IsFloat())
 		{
-			Value = FString::Printf(TEXT("%f"), Field.AsDouble());
+			const FCulturePtr LocaleCulture = FInternationalization::Get().GetCulture(FPlatformMisc::GetDefaultLocale());
+			const FNumberFormattingOptions Options = FNumberFormattingOptions()
+				.SetUseGrouping(false);
+			Value = FText::AsNumber(Field.AsDouble(), &Options, LocaleCulture).ToString();
 			Type = TEXT("Float");
 		}
 		// Check string
diff --git a/Engine/Source/Editor/StorageServerWidgets/Private/SBuildArtifactSelection.cpp b/Engine/Source/Editor/StorageServerWidgets/Private/SBuildArtifactSelection.cpp
--- a/Engine/Source/Editor/StorageServerWidgets/Private/SBuildArtifactSelection.cpp
+++ b/Engine/Source/Editor/StorageServerWidgets/Private/SBuildArtifactSelection.cpp
@@ -3,6 +3,7 @@
 #include "SBuildArtifactSelection.h"
 
 #include "Framework/MultiBox/MultiBoxBuilder.h"
+#include "HAL/PlatformApplicationMisc.h"
 #include "Framework/Docking/TabManager.h"
 #include "Styling/StyleColors.h"
 #include "Widgets/Input/SButton.h"
@@ -144,14 +145,6 @@ void SBuildArtifactSelection::Refresh(const FStringView InNamespace, const FStri
 		NamedArtifacts.Empty();
 	}
 
-	for (const TPair<FString, FArtifact>& NameArtifactPair : NamedArtifacts)
-	{
-		if (NameArtifactPair.Value.BuildRecord.BuildId == HighlightedBuildId)
-		{
-			SelectedArtifacts.AddUnique(NameArtifactPair.Key);
-		}
-	}
-
 	RebuildSelectedGroupArtifactHeader();
 	RebuildSelectedGroupArtifactContents();
 }
@@ -353,6 +346,35 @@ void SBuildArtifactSelection::RebuildSelectedGroupArtifactContents()
 						SNew(STextBlock)
 						.Justification(ETextJustify::Left)
 						.Text(FText::FromString(ArtifactName))
+						.ToolTipText_Lambda([this, ArtifactName]() -> FText
+						{
+							if (const FArtifact* Artifact = NamedArtifacts.Find(ArtifactName))
+							{
+								FCbFieldView TotalSizeField = Artifact->BuildRecord.Metadata["totalSize"];
+								if (TotalSizeField.HasValue() && !TotalSizeField.HasError())
+								{
+									uint64 TotalSize = 0;
+									if (TotalSizeField.IsInteger())
+									{
+										TotalSize = TotalSizeField.AsUInt64();
+									}
+									else if (TotalSizeField.IsFloat())
+									{
+										TotalSize = static_cast<uint64>(TotalSizeField.AsDouble());
+									}
+									else if (TotalSizeField.IsString())
+									{
+										LexFromString(TotalSize, *FString(FUTF8ToTCHAR(TotalSizeField.AsString())));
+									}
+									else
+									{
+										return FText::GetEmpty();
+									}
+									return FText::FromString(SBuildArtifactContentsViewer::FormatSize(TotalSize));
+								}
+							}
+							return FText::GetEmpty();
+						})
 						.Font_Lambda([this, ArtifactName, RegularFont, HighlightFont]
 						{
 							if (FArtifact* Artifact = NamedArtifacts.Find(ArtifactName))
@@ -364,6 +386,17 @@ void SBuildArtifactSelection::RebuildSelectedGroupArtifactContents()
 							}
 							return RegularFont;
 						})
+						.ColorAndOpacity_Lambda([this, ArtifactName]
+						{
+							if (FArtifact* Artifact = NamedArtifacts.Find(ArtifactName))
+							{
+								if (Artifact->BuildRecord.BuildId == HighlightedBuildId)
+								{
+									return FSlateColor(FStyleColors::AccentYellow);
+								}
+							}
+							return FSlateColor::UseForeground();
+						})
 					]
 				]
 				+SHorizontalBox::Slot()
@@ -433,6 +466,33 @@ TSharedRef<SWidget> SBuildArtifactSelection::CreateArtifactMenuWidget(UE::BuildS
 		EUserInterfaceActionType::Button
 	);
 
+	MenuBuilder.AddMenuEntry(
+		LOCTEXT("BuildSelection_CopyDownloadUrl", "Copy Download URL"),
+		FText(),
+		FSlateIcon(),
+		FUIAction(
+			FExecuteAction::CreateLambda([this, Artifact]
+				{
+					if (Artifact)
+					{
+						TSharedPtr<UE::Zen::Build::FBuildServiceInstance> ServiceInstance = BuildServiceInstance.Get();
+						if (ServiceInstance)
+						{
+							// TODO: This URL format is specific to UE Build Storage - we should genericize this
+							FString Url = FString::Printf(TEXT("https://%s/api/v2/builds/%s/%s/%s"),
+								*FString(ServiceInstance->GetEffectiveDomain()),
+								*Namespace,
+								*Artifact->BuildRecord.BucketId,
+								*FString(WriteToString<64>(Artifact->BuildRecord.BuildId)));
+							FPlatformApplicationMisc::ClipboardCopy(*Url);
+						}
+					}
+				})
+		),
+		NAME_None,
+		EUserInterfaceActionType::Button
+	);
+
 	MenuBuilder.AddMenuEntry(
 		LOCTEXT("BuildSelection_ViewMetadata", "View Metadata"),
 		FText(),
diff --git a/Engine/Source/Editor/StorageServerWidgets/Private/SBuildSelection.cpp b/Engine/Source/Editor/StorageServerWidgets/Private/SBuildSelection.cpp
--- a/Engine/Source/Editor/StorageServerWidgets/Private/SBuildSelection.cpp
+++ b/Engine/Source/Editor/StorageServerWidgets/Private/SBuildSelection.cpp
@@ -9,6 +9,7 @@
 #include "HAL/PlatformApplicationMisc.h"
 #include "HAL/PlatformFileManager.h"
 #include "Internationalization/FastDecimalFormat.h"
+#include "Internationalization/Internationalization.h"
 #include "Internationalization/TextChronoFormatter.h"
 #include "Math/BasicMathExpressionEvaluator.h"
 #include "Math/UnitConversion.h"
@@ -979,7 +980,7 @@ bool SBuildSelection::DoesPassAllViewFilters(TSharedPtr<FBuildGroup> InItem) con
 		return true;
 	}
 
-	if (FText::AsDateTime(InItem->CreatedAt, EDateTimeStyle::Short).ToString().Contains(SelectedTextFilter))
+	if (FText::AsDateTime(InItem->CreatedAt, EDateTimeStyle::Short, EDateTimeStyle::Short, FString(), FInternationalization::Get().GetCulture(FPlatformMisc::GetDefaultLocale())).ToString().Contains(SelectedTextFilter))
 	{
 		return true;
 	}
@@ -2430,7 +2431,7 @@ SBuildGroupTableRow::GenerateWidgetForColumn(const FName& ColumnName)
 				+ SHorizontalBox::Slot()
 				.Padding(1.f,1.f)
 				[
-					SNew(STextBlock).Text(FText::AsDateTime(BuildGroup->CreatedAt, EDateTimeStyle::Short))
+					SNew(STextBlock).Text(FText::AsDateTime(BuildGroup->CreatedAt, EDateTimeStyle::Short, EDateTimeStyle::Short, FString(), FInternationalization::Get().GetCulture(FPlatformMisc::GetDefaultLocale())))
 				];
 		}
 	}
diff --git a/Engine/Source/Programs/BuildStorageTool/BuildStorageTool.Build.cs b/Engine/Source/Programs/BuildStorageTool/BuildStorageTool.Build.cs
--- a/Engine/Source/Programs/BuildStorageTool/BuildStorageTool.Build.cs
+++ b/Engine/Source/Programs/BuildStorageTool/BuildStorageTool.Build.cs
@@ -26,6 +26,7 @@ public class BuildStorageTool : ModuleRules
 				"StandaloneRenderer",
 				"StorageServerWidgets",
 				"StudioTelemetry",
+				"ToolWidgets",
 				"Zen",
 			});
 
diff --git a/Engine/Source/Programs/BuildStorageTool/Private/BuildStorageTool.cpp b/Engine/Source/Programs/BuildStorageTool/Private/BuildStorageTool.cpp
--- a/Engine/Source/Programs/BuildStorageTool/Private/BuildStorageTool.cpp
+++ b/Engine/Source/Programs/BuildStorageTool/Private/BuildStorageTool.cpp
@@ -40,7 +40,12 @@
 #include "StudioTelemetry.h"
 #include "ZenBuildUtils.h"
 #include "BuildStorageToolIPC.h"
+#include "Framework/Application/IInputProcessor.h"
+#include "Framework/Commands/Commands.h"
+#include "HAL/PlatformApplicationMisc.h"
 #include "String/ParseTokens.h"
+#include "Widgets/Input/SEditableTextBox.h"
+#include "Widgets/Layout/SBox.h"
 
 
 #if PLATFORM_WINDOWS
@@ -78,6 +83,72 @@ static void HideOnCloseOverride(const TSharedRef<SWindow>& WindowBeingClosed)
 	WindowBeingClosed->HideWindow();
 }
 
+class FBuildStorageToolCommands : public TCommands<FBuildStorageToolCommands>
+{
+public:
+	FBuildStorageToolCommands()
+		: TCommands<FBuildStorageToolCommands>(
+			TEXT("BuildStorageTool"),
+			LOCTEXT("BuildStorageToolCommands", "Build Storage Tool"),
+			NAME_None,
+			FCoreStyle::Get().GetStyleSetName())
+	{
+	}
+
+	virtual void RegisterCommands() override
+	{
+		UI_COMMAND(GoToBuild, "Go to Build...", "Navigate to a build by entering its download URL", EUserInterfaceActionType::Button, FInputChord(EModifierKey::Control, EKeys::G));
+	}
+
+	TSharedPtr<FUICommandInfo> GoToBuild;
+};
+
+class FBuildStorageToolPasteHandler : public IInputProcessor
+{
+public:
+	void SetBuildSelection(TSharedPtr<SBuildSelection> InBuildSelection)
+	{
+		BuildSelection = InBuildSelection;
+	}
+
+	virtual void Tick(const float DeltaTime, FSlateApplication& SlateApp, TSharedRef<ICursor> Cursor) override {}
+
+	virtual bool HandleKeyDownEvent(FSlateApplication& SlateApp, const FKeyEvent& InKeyEvent) override
+	{
+		if (!(InKeyEvent.IsControlDown() || InKeyEvent.IsCommandDown()) || InKeyEvent.GetKey() != EKeys::V)
+		{
+			return false;
+		}
+
+		if (SlateApp.GetActiveModalWindow().IsValid())
+		{
+			return false;
+		}
+
+		TSharedPtr<SBuildSelection> PinnedBuildSelection = BuildSelection.Pin();
+		if (!PinnedBuildSelection.IsValid())
+		{
+			return false;
+		}
+
+		FString ClipboardText;
+		FPlatformApplicationMisc::ClipboardPaste(ClipboardText);
+		UE::Zen::Build::FBuildReference BuildReference;
+		if (UE::Zen::Build::TryParseBuildReferenceFromUrl(ClipboardText, BuildReference))
+		{
+			PinnedBuildSelection->ActOnBuildArtifact(SBuildSelection::EBuildArtifactAction::Highlight, BuildReference);
+			return true;
+		}
+
+		return false;
+	}
+
+	virtual const TCHAR* GetDebugName() const override { return TEXT("BuildStorageToolPasteHandler"); }
+
+private:
+	TWeakPtr<SBuildSelection> BuildSelection;
+};
+
 class FBuildStorageToolApp
 {
 	void ExitTool()
@@ -120,6 +191,103 @@ class FBuildStorageToolApp
 		FSlateApplication::Get().AddModalWindow(WindowDialog.ToSharedRef(), ParentWidget);
 	}
 
+	void GoToBuild()
+	{
+		FSlateApplication& SlateApp = FSlateApplication::Get();
+
+		if (SlateApp.GetActiveModalWindow() != nullptr)
+		{
+			return;
+		}
+
+		TSharedPtr<SWindow> DialogWindow;
+		bool bConfirmed = false;
+
+		TSharedRef<SEditableTextBox> UrlTextBox = SNew(SEditableTextBox)
+			.HintText(LOCTEXT("GoToBuildHint", "Enter a build download URL..."))
+			.OnTextCommitted_Lambda([&bConfirmed, &DialogWindow](const FText&, ETextCommit::Type CommitType)
+			{
+				if (CommitType == ETextCommit::OnEnter)
+				{
+					bConfirmed = true;
+					if (DialogWindow.IsValid())
+					{
+						DialogWindow->RequestDestroyWindow();
+					}
+				}
+			});
+
+		DialogWindow = SNew(SWindow)
+			.Title(LOCTEXT("GoToBuildTitle", "Go to Build"))
+			.SupportsMaximize(false)
+			.SupportsMinimize(false)
+			.SizingRule(ESizingRule::Autosized)
+			[
+				SNew(SBox)
+				.MinDesiredWidth(500.0f)
+				.MaxDesiredWidth(500.0f)
+				.Padding(FMargin(16.0f))
+				[
+					SNew(SVerticalBox)
+					+ SVerticalBox::Slot()
+					.AutoHeight()
+					.Padding(0, 0, 0, 8)
+					[
+						SNew(STextBlock)
+						.Text(LOCTEXT("GoToBuildLabel", "Enter the download URL of the build to navigate to:"))
+					]
+					+ SVerticalBox::Slot()
+					.AutoHeight()
+					.Padding(0, 0, 0, 16)
+					[
+						UrlTextBox
+					]
+					+ SVerticalBox::Slot()
+					.AutoHeight()
+					.HAlign(HAlign_Right)
+					[
+						SNew(SHorizontalBox)
+						+ SHorizontalBox::Slot()
+						.AutoWidth()
+						.Padding(0, 0, 8, 0)
+						[
+							SNew(SButton)
+							.Text(LOCTEXT("GoButton", "Go"))
+							.OnClicked_Lambda([&bConfirmed, &DialogWindow]()
+							{
+								bConfirmed = true;
+								DialogWindow->RequestDestroyWindow();
+								return FReply::Handled();
+							})
+						]
+						+ SHorizontalBox::Slot()
+						.AutoWidth()
+						[
+							SNew(SButton)
+							.Text(LOCTEXT("CancelButton", "Cancel"))
+							.OnClicked_Lambda([&DialogWindow]()
+							{
+								DialogWindow->RequestDestroyWindow();
+								return FReply::Handled();
+							})
+						]
+					]
+				]
+			];
+
+		SlateApp.AddModalWindow(DialogWindow.ToSharedRef(), Window);
+
+		if (bConfirmed)
+		{
+			FString Url = UrlTextBox->GetText().ToString();
+			UE::Zen::Build::FBuildReference BuildReference;
+			if (UE::Zen::Build::TryParseBuildReferenceFromUrl(Url, BuildReference))
+			{
+				BuildSelection->ActOnBuildArtifact(SBuildSelection::EBuildArtifactAction::Highlight, BuildReference);
+			}
+		}
+	}
+
 	bool CanExecuteExclusiveAction()
 	{
 		return !bLatentExclusiveOperationActive;
@@ -127,7 +295,7 @@ class FBuildStorageToolApp
 
 	TSharedRef< SWidget > MakeMainMenu()
 	{
-		FMenuBarBuilder MenuBuilder( nullptr );
+		FMenuBarBuilder MenuBuilder( CommandList );
 		{
 			MenuBuilder.AddPullDownMenu(
 				LOCTEXT( "FileMenu", "File" ),
@@ -208,7 +376,7 @@ class FBuildStorageToolApp
 		const bool bRecursivelySearchable = false;
 
 		FMenuBuilder MenuBuilder(true,
-			nullptr,
+			CommandList,
 			TSharedPtr<FExtender>(),
 			bCloseSelfOnly,
 			&FCoreStyle::Get(),
@@ -216,6 +384,10 @@ class FBuildStorageToolApp
 			NAME_None,
 			bRecursivelySearchable);
 
+		MenuBuilder.AddMenuEntry(FBuildStorageToolCommands::Get().GoToBuild);
+
+		MenuBuilder.AddMenuSeparator();
+
 		MenuBuilder.AddMenuEntry(
 					LOCTEXT("ExploreLogDirectory", "Explore Log Directory"),
 					LOCTEXT("ExploreLogDirectory_ToolTip", "Explore the directory containing log files"),
@@ -451,6 +623,13 @@ public:
 	{
 		ZenServiceInstanceManager = MakeShared<UE::Zen::FServiceInstanceManager>();
 		BuildServiceInstanceManager = MakeShared<UE::Zen::Build::FServiceInstanceManager>();
+
+		FBuildStorageToolCommands::Register();
+		CommandList = MakeShared<FUICommandList>();
+		CommandList->MapAction(
+			FBuildStorageToolCommands::Get().GoToBuild,
+			FExecuteAction::CreateRaw(this, &FBuildStorageToolApp::GoToBuild));
+
 		InstallMessageDialogOverride();
 		MessageListenServer.OnMessageReceived.AddRaw(this, &FBuildStorageToolApp::OnIPCMessageReceived);
 		MessageListenServer.Start();
@@ -460,6 +639,7 @@ public:
 	{
 		MessageListenServer.ShutdownJoin();
 		RemoveMessageDialogOverride();
+		FBuildStorageToolCommands::Unregister();
 	}
 
 	void Run()
@@ -601,6 +781,11 @@ public:
 			}
 		}
 
+		// Register input pre-processor for paste-to-navigate
+		TSharedRef<FBuildStorageToolPasteHandler> PasteHandler = MakeShared<FBuildStorageToolPasteHandler>();
+		PasteHandler->SetBuildSelection(BuildSelection);
+		Slate.RegisterInputPreProcessor(PasteHandler);
+
 		// Setting focus seems to have to happen after the Window has been added
 		Slate.ClearKeyboardFocus(EFocusCause::Cleared);
 
@@ -631,6 +816,8 @@ public:
 			GFrameCounter++;
 		}
 
+		Slate.UnregisterInputPreProcessor(PasteHandler);
+
 		// Make sure the window is hidden, because it might take a while for the background thread to finish.
 		Window->HideWindow();
 	}
@@ -748,6 +935,7 @@ private:
 	bool bVerify = false;
 	bool bBoostWorkers = false;
 
+	TSharedPtr<FUICommandList> CommandList;
 	std::atomic<bool> bLatentExclusiveOperationActive = false;
 	TSharedPtr<SBuildSelection> BuildSelection;
 	TSharedPtr<SBuildActivity> BuildActivity;
