BuildStorageTool: Exposing the "Append Files" transfer option to the menu UI.

#Synced-CL 51529686

diff --git a/Engine/Source/Programs/BuildStorageTool/Private/BuildStorageTool.cpp b/Engine/Source/Programs/BuildStorageTool/Private/BuildStorageTool.cpp
--- a/Engine/Source/Programs/BuildStorageTool/Private/BuildStorageTool.cpp
+++ b/Engine/Source/Programs/BuildStorageTool/Private/BuildStorageTool.cpp
@@ -417,9 +417,24 @@ class FBuildStorageToolApp
 				LOCTEXT("HostOverride_ToolTip", "Pick a host to use for build transfers instead of automatic selection of the nearest host"),
 				FNewMenuDelegate::CreateRaw(this, &FBuildStorageToolApp::OpenHostOverrideMenu));
 
+			MenuBuilder.AddMenuEntry(
+				LOCTEXT("AppendFiles", "Append Files Only"),
+				LOCTEXT("AppendFiles_ToolTip", "When downloading files to a directory, only append new files, don't remove files that used to be part of a build in that directory"),
+				FSlateIcon(),
+				FUIAction(
+					FExecuteAction::CreateRaw(
+						this, &FBuildStorageToolApp::ToggleAppendFiles),
+					FCanExecuteAction(),
+					FIsActionChecked::CreateRaw(
+						this, &FBuildStorageToolApp::IsAppendFilesSelected)
+				),
+				NAME_None,
+				EUserInterfaceActionType::Check
+			);
+
 			MenuBuilder.AddMenuEntry(
 				LOCTEXT("CleanDestination", "Clean Extra Files"),
-				LOCTEXT("CleanDestination_ToolTip", "When downloading files to a directory, any extra files that existed in that directory will be deleted"),
+				LOCTEXT("CleanDestination_ToolTip", "When downloading files to a directory, any extra files (not belonging to a past build) that existed in that directory will be deleted"),
 				FSlateIcon(),
 				FUIAction(
 					FExecuteAction::CreateRaw(
@@ -696,6 +711,10 @@ public:
 							{
 								InHostOverride = HostOverride;
 								RequestFlags = UE::Zen::Build::EBuildTransferRequestFlags::None;
+								if (bAppendFiles)
+								{
+									RequestFlags |= UE::Zen::Build::EBuildTransferRequestFlags::Append;
+								}
 								if (bCleanDestination)
 								{
 									RequestFlags |= UE::Zen::Build::EBuildTransferRequestFlags::Clean;
@@ -855,6 +874,16 @@ private:
 		FCoreDelegates::ModalMessageDialog.Unbind();
 	}
 
+	void ToggleAppendFiles()
+	{
+		bAppendFiles = !bAppendFiles;
+	}
+
+	bool IsAppendFilesSelected()
+	{
+		return bAppendFiles;
+	}
+
 	void ToggleCleanDestination()
 	{
 		bCleanDestination = !bCleanDestination;
@@ -929,6 +958,7 @@ private:
 	TSharedPtr<UE::Zen::Build::FServiceInstanceManager> BuildServiceInstanceManager;
 	const FBuildStorageToolParameters& ToolParameters;
 	FString HostOverride;
+	bool bAppendFiles = false;
 	bool bCleanDestination = false;
 	bool bForceRedownload = false;
 	bool bScavenging = true;
