BuildStorageTool: Bundle pngs needed for appicon with windows builds by declaring them as a runtime dependency.

#Synced-CL 51529686

diff --git a/Engine/Source/Programs/BuildStorageTool/BuildStorageTool.Build.cs b/Engine/Source/Programs/BuildStorageTool/BuildStorageTool.Build.cs
--- a/Engine/Source/Programs/BuildStorageTool/BuildStorageTool.Build.cs
+++ b/Engine/Source/Programs/BuildStorageTool/BuildStorageTool.Build.cs
@@ -53,6 +53,7 @@ public class BuildStorageTool : ModuleRules
 			RuntimeDependencies.Add("$(EngineDir)/Binaries/Win64/zenserver.exe");
 			RuntimeDependencies.Add("$(EngineDir)/Binaries/Win64/zenserver.pdb");
 			RuntimeDependencies.Add("$(EngineDir)/Binaries/DotNET/OidcToken/win-x64/OidcToken.exe");
+			RuntimeDependencies.Add("$(EngineDir)/Programs/BuildStorageTool/Resources/Windows/...*.png", StagedFileType.UFS);
 		}
 		else if (Target.Platform == UnrealTargetPlatform.Linux)
 		{
