BuildStorageTool: Undo custom appicon until it can be configured to stage correctly.

#Synced-CL 51529686

diff --git a/Engine/Programs/BuildStorageTool/Config/DefaultGame.ini b/Engine/Programs/BuildStorageTool/Config/DefaultGame.ini
--- a/Engine/Programs/BuildStorageTool/Config/DefaultGame.ini
+++ b/Engine/Programs/BuildStorageTool/Config/DefaultGame.ini
@@ -1,7 +1,6 @@
 [Staging]
 +AllowedDirectories=Engine/Content/Editor/Slate
 +AllowedDirectories=Engine/Content/Slate
-+AllowedDirectories=Engine/Programs/BuildStorageTool/Resources/Windows
 +AllowedConfigFiles=Engine/Config/BaseEditorPerProjectUserSettings.ini
 +AllowedConfigFiles=BuildStorageTool/Config/DeployId.ini
 
diff --git a/Engine/Source/Programs/BuildStorageTool/BuildStorageTool.Build.cs b/Engine/Source/Programs/BuildStorageTool/BuildStorageTool.Build.cs
--- a/Engine/Source/Programs/BuildStorageTool/BuildStorageTool.Build.cs
+++ b/Engine/Source/Programs/BuildStorageTool/BuildStorageTool.Build.cs
@@ -53,7 +53,6 @@ public class BuildStorageTool : ModuleRules
 			RuntimeDependencies.Add("$(EngineDir)/Binaries/Win64/zenserver.exe");
 			RuntimeDependencies.Add("$(EngineDir)/Binaries/Win64/zenserver.pdb");
 			RuntimeDependencies.Add("$(EngineDir)/Binaries/DotNET/OidcToken/win-x64/OidcToken.exe");
-			RuntimeDependencies.Add("$(EngineDir)/Programs/BuildStorageTool/Resources/Windows/...*.png", StagedFileType.UFS);
 		}
 		else if (Target.Platform == UnrealTargetPlatform.Linux)
 		{
diff --git a/Engine/Source/Programs/BuildStorageTool/Private/BuildStorageToolStyle.cpp b/Engine/Source/Programs/BuildStorageTool/Private/BuildStorageToolStyle.cpp
--- a/Engine/Source/Programs/BuildStorageTool/Private/BuildStorageToolStyle.cpp
+++ b/Engine/Source/Programs/BuildStorageTool/Private/BuildStorageToolStyle.cpp
@@ -59,13 +59,13 @@ TSharedRef< FSlateStyleSet > FBuildStorageToolStyle::Create()
 	const FVector2D Icon12x12(12.0f, 12.0f);
 	const FVector2D Icon16x16(16.0f, 16.0f);
 
-	{
-		const FString IconPath = FPaths::EngineSourceDir() / TEXT("Programs/BuildStorageTool/Resources/Windows/icon_24x24.png");
-		Style.Set("AppIcon", new FSlateImageBrush(IconPath, FVector2D(24, 24)));
-		Style.Set("AppIcon.Small", new FSlateImageBrush(IconPath, FVector2D(24, 24)));
-		Style.Set("AppIconPadding", FMargin(5.f, 5.f, 5.f, 5.f));
-		Style.Set("AppIconPadding.Small", FMargin(4.f, 4.f, 4.f, 4.f));
-	}
+	// {
+	// 	const FString IconPath = FPaths::EngineDir() / TEXT("Programs/BuildStorageTool/Resources/Windows/icon_24x24.png");
+	// 	Style.Set("AppIcon", new FSlateImageBrush(IconPath, FVector2D(24, 24)));
+	// 	Style.Set("AppIcon.Small", new FSlateImageBrush(IconPath, FVector2D(24, 24)));
+	// 	Style.Set("AppIconPadding", FMargin(5.f, 5.f, 5.f, 5.f));
+	// 	Style.Set("AppIconPadding.Small", FMargin(4.f, 4.f, 4.f, 4.f));
+	// }
 
 	Style.Set("Zen.FolderExplore", new EDITOR_IMAGE_BRUSH_SVG("Starship/Common/ContentBrowser", Icon12x12));
 	Style.Set("Zen.BrowseContent", new EDITOR_IMAGE_BRUSH_SVG("Starship/Common/BrowseContent", Icon16x16));
