BuildStorageTool: fixes for a task scope issue that affected Mac shutdown sequence.

#Synced-CL 51953975

diff --git a/Engine/Source/Programs/BuildStorageTool/Private/BuildStorageTool.cpp b/Engine/Source/Programs/BuildStorageTool/Private/BuildStorageTool.cpp
--- a/Engine/Source/Programs/BuildStorageTool/Private/BuildStorageTool.cpp
+++ b/Engine/Source/Programs/BuildStorageTool/Private/BuildStorageTool.cpp
@@ -980,6 +980,8 @@ int BuildStorageToolMain(const TCHAR* CmdLine)
 {
 	using namespace UE::BuildStorageTool;
 
+	const FTaskTagScope Scope(ETaskTag::EGameThread);
+
 	ON_SCOPE_EXIT
 	{
 		RequestEngineExit(TEXT("Exiting"));
@@ -988,8 +990,6 @@ int BuildStorageToolMain(const TCHAR* CmdLine)
 		FEngineLoop::AppExit();
 	};
 
-	const FTaskTagScope Scope(ETaskTag::EGameThread);
-
 	// start up the main loop
 	GEngineLoop.PreInit(CmdLine);
 
@@ -1093,10 +1093,21 @@ int main(int argc, char* argv[])
 #elif PLATFORM_MAC
 int32 MacMainWrapper(const TCHAR* CommandLine)
 {
+	// Clear EStaticInit on the game thread before FTaskTagScope(EGameThread) is constructed
+	// in BuildStorageToolMain. This prevents the FTaskTagScope constructor from calling
+	// GetStaticThreadId() on the game thread, which would cause IsRunningDuringStaticInit()
+	// to return false on the main thread during __cxa_finalize_ranges at process shutdown.
+	FTaskTagScope::SetTagNone();
 	return BuildStorageToolMain(CommandLine) ? 0 : 1;
 };
 int main(int argc, char* argv[])
 {
+	// Set the game thread scope on the main thread before the game thread starts. This
+	// initializes GetStaticThreadId() to the main thread's ID, which is required for
+	// IsRunningDuringStaticInit() to return true during __cxa_finalize_ranges at process
+	// shutdown (applicationWillTerminate: calls SetTagStaticInit() on the main thread, and
+	// IsRunningDuringStaticInit() checks that GetStaticThreadId() matches the current thread).
+	const FTaskTagScope MainThreadScope(ETaskTag::EGameThread);
 	[MacProgramDelegate mainWithArgc : argc argv : argv programMain : MacMainWrapper programExit : FEngineLoop::AppExit] ;
 }
 #else
