BuildStorageTool: More adjustments to task scopes to address ensure failures on Mac.

#Synced-CL 51953975

diff --git a/Engine/Source/Programs/BuildStorageTool/Private/BuildStorageTool.cpp b/Engine/Source/Programs/BuildStorageTool/Private/BuildStorageTool.cpp
--- a/Engine/Source/Programs/BuildStorageTool/Private/BuildStorageTool.cpp
+++ b/Engine/Source/Programs/BuildStorageTool/Private/BuildStorageTool.cpp
@@ -980,8 +980,6 @@ int BuildStorageToolMain(const TCHAR* CmdLine)
 {
 	using namespace UE::BuildStorageTool;
 
-	const FTaskTagScope Scope(ETaskTag::EGameThread);
-
 	ON_SCOPE_EXIT
 	{
 		RequestEngineExit(TEXT("Exiting"));
@@ -1078,12 +1076,14 @@ int BuildStorageToolMain(const TCHAR* CmdLine)
 int WINAPI WinMain(_In_ HINSTANCE hCurrInstance, _In_opt_ HINSTANCE hPrevInstance, _In_ LPSTR lpCmdLine, _In_ int nShowCmd)
 {
 	hInstance = hCurrInstance;
+	const FTaskTagScope Scope(ETaskTag::EGameThread);
 	return BuildStorageToolMain(GetCommandLineW())? 0 : 1;
 }
 
 #elif PLATFORM_LINUX
 int UnixMainWrapper(const TCHAR* CommandLine)
 {
+	const FTaskTagScope Scope(ETaskTag::EGameThread);
 	return BuildStorageToolMain(CommandLine) ? 0 : 1;
 };
 int main(int argc, char* argv[])
@@ -1093,21 +1093,19 @@ int main(int argc, char* argv[])
 #elif PLATFORM_MAC
 int32 MacMainWrapper(const TCHAR* CommandLine)
 {
-	// Clear EStaticInit on the game thread before FTaskTagScope(EGameThread) is constructed
-	// in BuildStorageToolMain. This prevents the FTaskTagScope constructor from calling
-	// GetStaticThreadId() on the game thread, which would cause IsRunningDuringStaticInit()
-	// to return false on the main thread during __cxa_finalize_ranges at process shutdown.
+	// On Mac the game thread starts with ActiveTaskTag == EStaticInit (thread-local default).
+	// FTaskTagScope(EGameThread) constructor would call GetStaticThreadId() from the EStaticInit
+	// path, seeding it to the game thread ID. This breaks IsRunningDuringStaticInit() on the
+	// main thread during __cxa_finalize_ranges, which needs GetStaticThreadId() to match the
+	// main thread. Calling SetTagNone() first skips the EStaticInit path in the constructor,
+	// leaving GetStaticThreadId() uninitialized until it is first called from the main thread
+	// (in IsRunningDuringStaticInit() during shutdown), where it correctly seeds to main thread ID.
 	FTaskTagScope::SetTagNone();
+	const FTaskTagScope Scope(ETaskTag::EGameThread);
 	return BuildStorageToolMain(CommandLine) ? 0 : 1;
 };
 int main(int argc, char* argv[])
 {
-	// Set the game thread scope on the main thread before the game thread starts. This
-	// initializes GetStaticThreadId() to the main thread's ID, which is required for
-	// IsRunningDuringStaticInit() to return true during __cxa_finalize_ranges at process
-	// shutdown (applicationWillTerminate: calls SetTagStaticInit() on the main thread, and
-	// IsRunningDuringStaticInit() checks that GetStaticThreadId() matches the current thread).
-	const FTaskTagScope MainThreadScope(ETaskTag::EGameThread);
 	[MacProgramDelegate mainWithArgc : argc argv : argv programMain : MacMainWrapper programExit : FEngineLoop::AppExit] ;
 }
 #else
