BuildStorageTool: Adding Mac icons

#Synced-CL 52199356

diff --git a/Engine/Source/Programs/BuildStorageTool/Resources/Mac/Assets.xcassets/AppIcon.appiconset/Contents.json b/Engine/Source/Programs/BuildStorageTool/Resources/Mac/Assets.xcassets/AppIcon.appiconset/Contents.json
new file mode 100644
--- /dev/null
+++ b/Engine/Source/Programs/BuildStorageTool/Resources/Mac/Assets.xcassets/AppIcon.appiconset/Contents.json
@@ -0,0 +1,68 @@
+{
+  "images" : [
+    {
+      "filename" : "icon_16x16.png",
+      "idiom" : "mac",
+      "scale" : "1x",
+      "size" : "16x16"
+    },
+    {
+      "filename" : "[REDACTED_EMAIL]",
+      "idiom" : "mac",
+      "scale" : "2x",
+      "size" : "16x16"
+    },
+    {
+      "filename" : "icon_32x32.png",
+      "idiom" : "mac",
+      "scale" : "1x",
+      "size" : "32x32"
+    },
+    {
+      "filename" : "[REDACTED_EMAIL]",
+      "idiom" : "mac",
+      "scale" : "2x",
+      "size" : "32x32"
+    },
+    {
+      "filename" : "icon_128x128.png",
+      "idiom" : "mac",
+      "scale" : "1x",
+      "size" : "128x128"
+    },
+    {
+      "filename" : "[REDACTED_EMAIL]",
+      "idiom" : "mac",
+      "scale" : "2x",
+      "size" : "128x128"
+    },
+    {
+      "filename" : "icon_256x256.png",
+      "idiom" : "mac",
+      "scale" : "1x",
+      "size" : "256x256"
+    },
+	{
+      "filename" : "[REDACTED_EMAIL]",
+      "idiom" : "mac",
+      "scale" : "2x",
+      "size" : "256x256"
+	},
+	{
+	  "filename" : "icon_512x512.png",
+	  "idiom" : "mac",
+	  "scale" : "1x",
+	  "size" : "512x512"
+	},
+	{
+	  "filename" : "[REDACTED_EMAIL]",
+	  "idiom" : "mac",
+	  "scale" : "2x",
+	  "size" : "512x512"
+    }
+  ],
+  "info" : {
+    "author" : "xcode",
+    "version" : 1
+  }
+}
diff --git a/Engine/Source/Programs/BuildStorageTool/Resources/Mac/Assets.xcassets/Contents.json b/Engine/Source/Programs/BuildStorageTool/Resources/Mac/Assets.xcassets/Contents.json
new file mode 100644
--- /dev/null
+++ b/Engine/Source/Programs/BuildStorageTool/Resources/Mac/Assets.xcassets/Contents.json
@@ -0,0 +1,6 @@
+{
+  "info" : {
+    "author" : "xcode",
+    "version" : 1
+  }
+}
